Engine source (Main branch up to CL 2027741)

diff --git a/Engine/Build/build.properties b/Engine/Build/build.properties
--- a/Engine/Build/build.properties
+++ b/Engine/Build/build.properties
@@ -1,3 +1,3 @@
-TimestampForBVT=2014-03-16_04.09.05
-EngineVersion=2027525
+TimestampForBVT=2014-03-16_20.07.09
+EngineVersion=2027699
 BranchName=++depot+UE4
diff --git a/Engine/Source/Developer/SimplygonMeshReduction/Private/SimplygonMeshReduction.cpp b/Engine/Source/Developer/SimplygonMeshReduction/Private/SimplygonMeshReduction.cpp
--- a/Engine/Source/Developer/SimplygonMeshReduction/Private/SimplygonMeshReduction.cpp
+++ b/Engine/Source/Developer/SimplygonMeshReduction/Private/SimplygonMeshReduction.cpp
@@ -334,6 +334,12 @@ public:
 		ANSICHAR VersionHash[200];
 
 #if WITH_SIMPLYGON_DLL
+		if (FRocketSupport::IsRocket())
+		{
+			// this was killing DDC build, for now we won't even try in rocket mode.
+			return NULL;
+		}
+
 		typedef void (*GetInterfaceVersionSimplygonSDKPtr)(ANSICHAR*);
 		typedef int (*InitializeSimplygonSDKPtr)(const char* LicenseData , SimplygonSDK::ISimplygonSDK** OutInterfacePtr);
  
diff --git a/README.md b/README.md
--- a/README.md
+++ b/README.md
@@ -1,12 +1,14 @@
 Welcome to the UE4 source code!
-=================================
+===============================
 
 This is the source code page for the **Unreal Engine on GitHub**.  With the UE4 source code, you can modify the
 engine and tools in any way imaginable and share your changes with others!
 
-You can compile for either Windows or Mac, as well as build runtime games for Android and iOS.  Source for 
+You can build the editor for Windows and Mac and compile games for both Android and iOS.  Source code for 
 all tools is included as well, such as Unreal Lightmass and Unreal Frontend.
 
+Before continuing, check out this [short tutorial video](http://youtu.be/usjlNHPn-jo) about getting started with the engine code.
+
 Have fun!!  We can't wait to see what you create!!
 
 
@@ -21,32 +23,33 @@ You can **download the source** in .zip format, but in order to *contribute to t
 **setup a Git client** on your computer.  [This page](http://help.github.com/articles/set-up-git) will have 
 the info you need.
 
-Remember, you'll need to download dependencies in order to actually build and run the engine.  You'll find those on 
+Remember, you'll need to _download dependencies_ in order to actually build and run the engine.  You'll find those on 
 the [releases page](https://github.com/EpicGames/UnrealEngine/releases) along with the source!
 
-We're also experimenting with publishing bleeding edge changes from our development teams!  This is the newest stuff
-we have, pushed into the [master branch](https://github.com/EpicGames/UnrealEngine/tree/master) on GitHub.  No guarantee 
-is made that the code in 'master' is stable or compiles.
+We're also publishing bleeding edge changes from our engine team!  This is the newest code we have, pushed into 
+the [master branch](https://github.com/EpicGames/UnrealEngine/tree/master) on GitHub.  We'll continue to update
+the 'master' code with fresh snapshots every so often, we can't guarantee the code is stable or even compiles.
 
 
 
 Getting up and running
 ----------------------
 
-1. **Create a folder** on your computer for the code and dependencies.
+1. **Download the source** and unzip it to a folder, or [create a fork](https://github.com/EpicGames/UnrealEngine/tree/4.0.1-release)
+   and **clone the repository**.  If you clone, don't forget to switch to the correct branch for this release!  (The 'master' branch 
+   has unstable code, so you will want to make sure to choose a release branch.)
 
-1. **Download the source** and unzip it into your folder, or [create a fork](https://github.com/EpicGames/UnrealEngine/tree/4.0.1-release)
-   and clone it in.  If you clone, don't forget to switch to the correct branch for this release!  (The 'master' branch has unstable code, 
-   so you will want to make sure to choose a release branch.)
+1. You should now an _UnrealEngine_ folder on your computer.  All of the source and dependencies will go into this folder.  The folder name might 
+   have a branch suffix (such as _UnrealEngine-4.0_), but that totally fine.
 
 1. Download the **required dependencies** files for the [latest release](https://github.com/EpicGames/UnrealEngine/releases/tag/4.0.1-release): 
    [Required_1of2.zip](https://github.com/EpicGames/UnrealEngine/releases/download/4.0.1-release/Required_1of2.zip), 
    [Required_2of2.zip](https://github.com/EpicGames/UnrealEngine/releases/download/4.0.1-release/Required_2of2.zip).
 
-1. Unzip the dependencies to the **same folder** that you saved the source.  Be careful to make sure the folders are merged together 
-   correctly.  On Mac, we recommend **Alt + dragging** the unzipped files into your source folder, then selecting 'Keep Newer' if prompted.
+1. Unzip the dependencies into the _UnrealEngine_ folder alongside the source.  Be careful to make sure the folders are merged together 
+   correctly.  On Mac, we recommend **Option + dragging** the unzipped files into the _UnrealEngine_ folder, then selecting **Keep Newer** if prompted.
 
-1. Okay, platform-specific stuff comes next.  See the sections below:
+1. Okay, platform stuff comes next.  Depending on whether you are on Windows or Mac, follow one of the sections below:
 
 
 ### Windows
@@ -56,47 +59,52 @@ Getting up and running
 
 1. Make sure you have [June 2010 DirectX runtime](http://www.microsoft.com/en-us/download/details.aspx?id=8109) installed.  You don't need the SDK, just the runtime.
 
-1. You'll need project files in order to compile.  Double-click **GenerateProjectFiles.bat**.  It will only take a few moments to complete.
+1. You'll need project files in order to compile.  In the _UnrealEngine_ folder, double-click on **GenerateProjectFiles.bat**.  It should take less than a minute to complete.
 
 1. Load the project into Visual Studio by double-clicking on the **UE4.sln** file.
 
-1. It's time to **compile the editor**!  In Visual Studio, make sure your solution configuration is set to **Development Editor**, and your solution platform is set to **Win64**.  Right click on the **UE4** target and select **Build**.  It will take between 15 and 40 minutes to finish compiling, depending on your system specs.
+1. It's time to **compile the editor**!  In Visual Studio, make sure your solution configuration is set to **Development Editor**, and your solution 
+   platform is set to **Win64**.  Right click on the **UE4** target and select **Build**.  It will take between 15 and 40 minutes to finish compiling,
+   depending on your system specs.
 
-1. After compiling finishes, you can **load the editor** from the IDE by setting your startup project to **UE4** and pressing **F5** to debug.
+1. After compiling finishes, you can **load the editor** from Visual Studio by setting your startup project to **UE4** and pressing **F5** to debug.
 
-1. Setup your Windows shell so that you can right click on .uproject files.  Find the file named **RegisterShellCommands.bat** in the _UnrealEngine/Engine/Build/BatchFiles/_ folder, and run it!
+1. One last thing.  You'll want to setup your Windows shell so that you can interact with .uproject files.  Find the file named **RegisterShellCommands.bat** in 
+   the _UnrealEngine/Engine/Build/BatchFiles/_ folder, and run it.  Now, you'll be able to double-click .uproject files to load the project, or right click on
+   .uprojects to quickly update Visual Studio files.
 
 
 ### Mac
 
 1. Be sure to have [Xcode 5.1](https://itunes.apple.com/us/app/xcode/id497799835) installed.
 
-1. You'll need project files in order to compile.  Double-click **GenerateProjectFiles.command**.  It will only take a few moments to complete.
+1. You'll need project files in order to compile.  In the _UnrealEngine_ folder, double-click on **GenerateProjectFiles.command**.  It should take less than a minute to complete.  You can close the Terminal window afterwards.
 
 1. Load the project into Xcode by double-clicking on the **UE4.xcodeproj** file.
 
-1. It's time to **compile the editor**!  In Xcode, build the **UE4Editor - Mac** target for **My Mac** by selecting it in the target drop down in the top left, then using **Product** -> **Build For** -> **Running**.  It will take between 15 and 40 minutes, depending on your system specs.
+1. It's time to **compile the editor**!  In Xcode, build the **UE4Editor - Mac** target for **My Mac** by selecting it in the target drop down
+   in the top left, then using **Product** -> **Build For** -> **Running**.  It will take between 15 and 40 minutes, depending on your system specs.
 
-1. After compiling finishes, you can **load the editor** by from the IDE using the **Product** -> **Run** menu command!
+1. After compiling finishes, you can **load the editor** from Xcode using the **Product** -> **Run** menu command!
 
 
 
-Notes
------
+A few other things
+------------------
 
-The first time you start the editor from a fresh source build, you may experience longer load times and out of date shaders.  This only happens on the first run as the engine optimizes content for the platform.
+The first time you start the editor from a fresh source build, you may experience long load times.  This only happens on the first 
+run as the engine optimizes content for the platform and _fills the derived data cache_.
 
 Visual Studio 2013 and Xcode 5.1 are the supported and recommended compilers to use.
 
-Visual Studio 2012 is also supported, but you'll need to make a code change and download the **optional** dependencies 
+Visual Studio 2012 is also supported, but you'll need to make a code change and download the _optional dependencies_
 on the [releases page](https://github.com/EpicGames/UnrealEngine/releases/tag/4.0.1-release).  See the Unreal Engine 
 [documentation](http://docs.unrealengine.com) for more details about using older versions of Visual Studio.
 
 You should probably always work with a versioned **release** branch.  The master branch contains unstable and possibly untested code,
 but it should be a great reference for new developments, or for spot merging bug fixes.  Use it at your own risk.  
-Full engine releases tend to lag behind latest development branch due to QA testing, but we wanted to make sure that 
-you have access to the very latest stuff we're doing.
 
-To build for platforms besides Windows and Mac, please see the Unreal Engine [documentation](http://docs.unrealengine.com).  Some platforms may require
-you to download the **optional** dependencies as well.
+To build for platforms besides Windows and Mac, please see the Unreal Engine [documentation](http://docs.unrealengine.com).  Android 
+development currently works best from a PC.  Conversely, iOS programming requires a Mac.  Some platforms may require
+you to download the _optional dependencies_ as well.
 
