Engine source (Main branch up to CL 2027741)

diff --git a/README.md b/README.md
--- a/README.md
+++ b/README.md
@@ -35,12 +35,16 @@ the 'master' code with fresh snapshots every so often, we can't guarantee the co
 Getting up and running
 ----------------------
 
+Here is the fun part!  This is a quick start guide to getting up and running with the source.  The steps below will take you through cloning your own private fork, then compiling and 
+running the editor yourself on Windows or Mac.  Oh, and you might want to watch our [short tutorial video](http://youtu.be/usjlNHPn-jo)
+first.  Okay, here we go!
+
 1. **Download the source** and unzip it to a folder, or [create a fork](https://github.com/EpicGames/UnrealEngine/tree/4.0.1-release)
    and **clone the repository**.  If you clone, don't forget to switch to the correct branch for this release!  (The 'master' branch 
    has unstable code, so you will want to make sure to choose a release branch.)
 
-1. You should now an _UnrealEngine_ folder on your computer.  All of the source and dependencies will go into this folder.  The folder name might 
-   have a branch suffix (such as _UnrealEngine-4.0_), but that totally fine.
+1. You should now have an _UnrealEngine_ folder on your computer.  All of the source and dependencies will go into this folder.  The folder name might 
+   have a branch suffix (such as _UnrealEngine-4.0_), but that's totally fine.
 
 1. Download the **required dependencies** files for the [latest release](https://github.com/EpicGames/UnrealEngine/releases/tag/4.0.1-release): 
    [Required_1of2.zip](https://github.com/EpicGames/UnrealEngine/releases/download/4.0.1-release/Required_1of2.zip), 
@@ -89,22 +93,26 @@ Getting up and running
 
 
 
-A few other things
-------------------
+More info
+---------
+
+Visual Studio 2013 and Xcode 5.1 are the supported and recommended compilers to use.
 
 The first time you start the editor from a fresh source build, you may experience long load times.  This only happens on the first 
 run as the engine optimizes content for the platform and _fills the derived data cache_.
 
-Visual Studio 2013 and Xcode 5.1 are the supported and recommended compilers to use.
+You should probably always work with a versioned **release** branch.  The master branch contains unstable and possibly untested code,
+but it should be a great reference for new developments, or for spot merging bug fixes.  Use it at your own risk.  
+
+To build for platforms besides Windows and Mac, please see the Unreal Engine [documentation](http://docs.unrealengine.com).  Android 
+development currently works best from a PC.  Conversely, iOS programming requires a Mac.  Some platforms may require
+you to download _optional dependencies_ or install _platform development SDKs_ as well.
 
 Visual Studio 2012 is also supported, but you'll need to make a code change and download the _optional dependencies_
 on the [releases page](https://github.com/EpicGames/UnrealEngine/releases/tag/4.0.1-release).  See the Unreal Engine 
 [documentation](http://docs.unrealengine.com) for more details about using older versions of Visual Studio.
 
-You should probably always work with a versioned **release** branch.  The master branch contains unstable and possibly untested code,
-but it should be a great reference for new developments, or for spot merging bug fixes.  Use it at your own risk.  
+Your private forks of the Unreal Engine code are associated with your GitHub account permissions.  Just remember
+that if you unsubscribe or switch GitHub user names, you'll need to re-fork and upload your changes from a local copy. 
 
-To build for platforms besides Windows and Mac, please see the Unreal Engine [documentation](http://docs.unrealengine.com).  Android 
-development currently works best from a PC.  Conversely, iOS programming requires a Mac.  Some platforms may require
-you to download the _optional dependencies_ as well.
 
