Engine source (Main branch up to CL 2027741)

diff --git a/README.md b/README.md
--- a/README.md
+++ b/README.md
@@ -74,7 +74,7 @@ first.  Okay, here we go!
 1. After compiling finishes, you can **load the editor** from Visual Studio by setting your startup project to **UE4** and pressing **F5** to debug.
 
 1. One last thing.  You'll want to setup your Windows shell so that you can interact with .uproject files.  Find the file named **RegisterShellCommands.bat** in 
-   the _UnrealEngine/Engine/Build/BatchFiles/_ folder, and run it.  Now, you'll be able to double-click .uproject files to load the project, or right click on
+   the _UnrealEngine/Engine/Build/BatchFiles/_ folder.  Right click on the file and select **run as Administrator**.  Now, you'll be able to double-click .uproject files to load the project, or right click on
    .uprojects to quickly update Visual Studio files.
 
 
