Engine source (Main branch up to CL 2027741)

diff --git a/.gitattributes b/.gitattributes
new file mode 100644
--- /dev/null
+++ b/.gitattributes
@@ -0,0 +1,12 @@
+#
+# This file contains rules that control how Git handles binary and text files, including line endings
+#
+ 
+
+# Make sure Windows batch files preserve CR/LF line endings, otherwise they may not be able to execute.  Windows
+# batch files require a CR/LF for labels to work properly, otherwise they may fail when labels straddle 512-byte
+# block boundaries.  This is important when files are downloaded through a zip archive that was authored on a
+# Linux machine (the default behavior on GitHub)
+
+*.bat eol=crlf
+
diff --git a/.gitignore b/.gitignore
--- a/.gitignore
+++ b/.gitignore
@@ -54,6 +54,47 @@
 # Ignore any saved local files
 /Engine/Saved/**
 
-# Ignore template projects
-/Templates/**
+# Ignore third party folders
+/Engine/Source/ThirdParty/ANGLE/**
+/Engine/Source/ThirdParty/DirectShow/**
+/Engine/Source/ThirdParty/elftoolchain/**
+/Engine/Source/ThirdParty/FaceBook/**
+/Engine/Source/ThirdParty/FBX/**
+/Engine/Source/ThirdParty/ForsythTriOO/**
+/Engine/Source/ThirdParty/FreeType2/**
+/Engine/Source/ThirdParty/HACD/**
+/Engine/Source/ThirdParty/hlslcc/**
+/Engine/Source/ThirdParty/ICU/**
+/Engine/Source/ThirdParty/IntelTBB/**
+/Engine/Source/ThirdParty/jemalloc/**
+/Engine/Source/ThirdParty/Kiss_FFT/**
+/Engine/Source/ThirdParty/libcurl/**
+/Engine/Source/ThirdParty/libJPG/**
+/Engine/Source/ThirdParty/libOpus/**
+/Engine/Source/ThirdParty/libPNG/**
+/Engine/Source/ThirdParty/Licenses/**
+/Engine/Source/ThirdParty/Mac/**
+/Engine/Source/ThirdParty/MCPP/**
+/Engine/Source/ThirdParty/nvtesslib/**
+/Engine/Source/ThirdParty/nvTextureTools/**
+/Engine/Source/ThirdParty/nvTriStrip/**
+/Engine/Source/ThirdParty/Oculus/**
+/Engine/Source/ThirdParty/Ogg/**
+/Engine/Source/ThirdParty/OpenAL/**
+/Engine/Source/ThirdParty/OpenGL/**
+/Engine/Source/ThirdParty/Perforce/**
+/Engine/Source/ThirdParty/PhysX/**
+/Engine/Source/ThirdParty/portmidi/**
+/Engine/Source/ThirdParty/Qualcomm/**
+/Engine/Source/ThirdParty/Recast/**
+/Engine/Source/ThirdParty/SDL/**
+/Engine/Source/ThirdParty/Vorbis/**
+/Engine/Source/ThirdParty/Windows/DirectX/**
+/Engine/Source/ThirdParty/zlib/**
+
+# Add back in all the build files for third party modules
+!/Engine/Source/ThirdParty/**/*.[Bb]uild.cs
 
+# Ignore sample and template projects
+/Samples/**
+/Templates/**
diff --git a/README.md b/README.md
--- a/README.md
+++ b/README.md
@@ -63,7 +63,7 @@ first.  Okay, here we go!
 
 1. Make sure you have [June 2010 DirectX runtime](http://www.microsoft.com/en-us/download/details.aspx?id=8109) installed.  You don't need the SDK, just the runtime.
 
-1. You'll need project files in order to compile.  In the _UnrealEngine_ folder, double-click on **GenerateProjectFiles.bat**.  It should take less than a minute to complete.
+1. You'll need project files in order to compile.  In the _UnrealEngine_ folder, double-click on **GenerateProjectFiles.bat**.  It should take less than a minute to complete.  On Windows 8, a warning from SmartScreen may appear.  Click "More info", then "Run anyway" to continue.
 
 1. Load the project into Visual Studio by double-clicking on the **UE4.sln** file.
 
