Code snapshot for CL 2063134

[CL 2063533 in Main branch]

diff --git a/.gitignore b/.gitignore
--- a/.gitignore
+++ b/.gitignore
@@ -89,6 +89,7 @@
 /Engine/Source/ThirdParty/Recast/**
 /Engine/Source/ThirdParty/SDL/**
 /Engine/Source/ThirdParty/SDL2/**
+/Engine/Source/ThirdParty/Steamworks/**
 /Engine/Source/ThirdParty/Vorbis/**
 /Engine/Source/ThirdParty/Windows/DirectX/**
 /Engine/Source/ThirdParty/zlib/**
