Minor: fix bashism in GenerateProjectFiles.sh

diff --git a/GenerateProjectFiles.sh b/GenerateProjectFiles.sh
--- a/GenerateProjectFiles.sh
+++ b/GenerateProjectFiles.sh
@@ -11,7 +11,7 @@ if [ ! -f Engine/Build/BatchFiles/Mac/GenerateProjectFiles.sh ]; then
   exit 1
 fi 
 
-if [ "$(uname)" == "Darwin" ]; then
+if [ "$(uname)" = "Darwin" ]; then
 	cd Engine/Build/BatchFiles/Mac
 	sh ./GenerateLLDBInit.sh
 	sh ./GenerateProjectFiles.sh $@
