Updated .gitignore to remove HTML5 TPS binaries

diff --git a/.gitignore b/.gitignore
--- a/.gitignore
+++ b/.gitignore
@@ -72,6 +72,7 @@
 /Engine/Source/ThirdParty/**/*.dylib
 /Engine/Source/ThirdParty/**/*.bmp
 /Engine/Source/ThirdParty/**/*.so
+/Engine/Source/ThirdParty/**/*.bc
 /Engine/Source/ThirdParty/**/*.bsc
 /Engine/Source/ThirdParty/**/*.idb
 
