Updated README for Master branch

diff --git a/README.md b/README.md
--- a/README.md
+++ b/README.md
@@ -19,21 +19,16 @@ Have fun!!  We can't wait to see what you create!!
 Source releases
 ---------------
 
-The latest version is [4.1.0](https://github.com/EpicGames/UnrealEngine/releases/tag/4.1.0-release).  Other 
-releases can be found [right here](https://github.com/EpicGames/UnrealEngine/releases).
+This branch mirrors the live codebase that we're working on at Epic, and contains [bleeding-edge changes](https://github.com/EpicGames/UnrealEngine/commits/master) 
+by our engine team. The head revision may be unstable or not even compile, so we encourage working from a versioned release branch which is fully 
+tested by our QA team instead. Official releases can be found [here](https://github.com/EpicGames/UnrealEngine/releases).
 
 You can **download the source** in .zip format, but in order to *contribute to the community* you should
 **setup a Git client** on your computer.  [This page](http://help.github.com/articles/set-up-git) will have 
 the info you need.
 
-Remember, you'll need to _download dependencies_ in order to actually build and run the engine.  You'll find those on 
-the [releases page](https://github.com/EpicGames/UnrealEngine/releases) along with the source!
-
-We're also publishing bleeding edge changes from our engine team!  This is the newest code we have, pushed into 
-the [master branch](https://github.com/EpicGames/UnrealEngine/tree/master) on GitHub.  You can see [live commits](https://github.com/EpicGames/UnrealEngine/commits/master)
-from Epic programmers along with integrated code submissions from the community!  Of course these commits have not 
-been QA tested yet and may not even compile.  We'll also frequently publish [preview builds](https://github.com/EpicGames/UnrealEngine/releases/tag/latest-preview) which
-contain required dependencies you'll need to compile master branch code.
+Remember, you'll need to _download dependencies_ in order to actually build and run the engine.  We periodically update the latest-preview tag with matching 
+dependencies on our [releases page](https://github.com/EpicGames/UnrealEngine/releases).
 
 
 
@@ -44,16 +39,15 @@ Here is the fun part!  This is a quick start guide to getting up and running wit
 running the editor yourself on Windows or Mac.  Oh, and you might want to watch our [short tutorial video](http://youtu.be/usjlNHPn-jo)
 first.  Okay, here we go!
 
-1. **Download the source** and unzip it to a folder, or [create a fork](https://github.com/EpicGames/UnrealEngine/tree/4.1.0-release)
-   and **clone the repository**.  If you clone, don't forget to switch to the correct branch for this release!  (The 'master' branch 
-   has unstable code, so you will want to make sure to choose a release branch.)
+1. [Download the source](https://github.com/EpicGames/UnrealEngine/archive/latest-preview.zip) and unzip it to a folder, or 
+   [create a fork](https://help.github.com/articles/fork-a-repo) and **clone the repository**.
 
 1. You should now have an _UnrealEngine_ folder on your computer.  All of the source and dependencies will go into this folder.  The folder name might 
-   have a branch suffix (such as _UnrealEngine-4.1_), but that's totally fine.
+   have a branch suffix (such as _UnrealEngine-master_), which is normal.
 
-1. Download the **required dependencies** files for the [latest release](https://github.com/EpicGames/UnrealEngine/releases/tag/4.1.0-release): 
-   [Required_1of2.zip](https://github.com/EpicGames/UnrealEngine/releases/download/4.1.0-release/Required_1of2.zip), 
-   [Required_2of2.zip](https://github.com/EpicGames/UnrealEngine/releases/download/4.1.0-release/Required_2of2.zip).
+1. Download the **required dependencies** files for the [latest preview tag](https://github.com/EpicGames/UnrealEngine/releases/tag/latest-preview): 
+   [Required_1of2.zip](https://github.com/EpicGames/UnrealEngine/releases/download/latest-preview/Required_1of2.zip), 
+   [Required_2of2.zip](https://github.com/EpicGames/UnrealEngine/releases/download/latest-preview/Required_2of2.zip).
 
 1. Unzip the dependencies into the _UnrealEngine_ folder alongside the source.  Be careful to make sure the folders are merged together 
    correctly.  On Mac, we recommend **Option + dragging** the unzipped files into the _UnrealEngine_ folder, then selecting **Keep Newer** if prompted.
@@ -98,26 +92,30 @@ first.  Okay, here we go!
 
 
 
-More info
----------
+### Additional target platforms
 
-Visual Studio 2013 and Xcode 5.1 are the supported and recommended compilers to use.
+**Android** development currently works best from a PC. See the [Android getting started guide](https://docs.unrealengine.com/latest/INT/Platforms/Android/GettingStarted/).
 
-The first time you start the editor from a fresh source build, you may experience long load times.  This only happens on the first 
-run as the engine optimizes content for the platform and _fills the derived data cache_.
+**iOS** programming requires a Mac. Instructions are in the [iOS getting started guide](https://docs.unrealengine.com/latest/INT/Platforms/iOS/GettingStarted/index.html).
 
-You should probably always work with a versioned **release** branch.  The master branch contains unstable and possibly untested code,
-but it should be a great reference for new developments, or for spot merging bug fixes.  Use it at your own risk.  
+**HTML5** is supported using Emscripten, and requires the [optional dependencies](https://github.com/EpicGames/UnrealEngine/releases/download/latest-preview/Optional.zip) to be unzipped into your _UnrealEngine_ folder. Please see the [HTML5 getting started guide](https://docs.unrealengine.com/latest/INT/Platforms/HTML5/GettingStarted/index.html).
 
-To build for platforms besides Windows and Mac, please see the Unreal Engine [documentation](http://docs.unrealengine.com).  Android 
-development currently works best from a PC.  Conversely, iOS programming requires a Mac.  Some platforms may require
-you to download _optional dependencies_ or install _platform development SDKs_ as well.
+**Linux** support requires the [optional dependencies](https://github.com/EpicGames/UnrealEngine/releases/download/latest-preview/Optional.zip) to be unzipped into your _UnrealEngine_ folder, and the [cross-compiler toolchain](http://cdn.unrealengine.com/qfe/v3_clang-3.3_ld-2.24_glibc-2.12.2.zip) to be installed. See the [Linux cross-compiling page on the wiki](https://wiki.unrealengine.com/Compiling_For_Linux).
 
-Visual Studio 2012 is also supported, but you'll need to make a code change and download the _optional dependencies_
-on the [releases page](https://github.com/EpicGames/UnrealEngine/releases/tag/4.1.0-release).  See the Unreal Engine 
-[documentation](http://docs.unrealengine.com) for more details about using older versions of Visual Studio.
+**Playstation 4** or **XboxOne** development require additional files that can only be provided after your registered developer status is confirmed by Sony or Microsoft. See [the announcement blog post](https://www.unrealengine.com/blog/playstation-4-and-xbox-one-now-supported) for more information.
 
-Your private forks of the Unreal Engine code are associated with your GitHub account permissions.  Just remember
-that if you unsubscribe or switch GitHub user names, you'll need to re-fork and upload your changes from a local copy. 
 
 
+More info
+---------
+
+Visual Studio 2013 and Xcode 5.1 are strongly recommended for development.
+
+Legacy support for Visual Studio 2012 can be enabled by unzipping the [optional dependencies](https://github.com/EpicGames/UnrealEngine/releases/download/4.3.0-preview/Optional.zip) into your UnrealEngine folder
+and changing WindowsPlatform.Compiler to VisualStudio2012 in Engine\Source\Programs\UnrealBuildTool\Windows\UEBuildWindows.cs.
+
+The first time you start the editor from a fresh source build, you may experience long load times.  This only happens on the first 
+run as the engine optimizes content for the platform and _fills the derived data cache_.
+
+Your private forks of the Unreal Engine code are associated with your GitHub account permissions.  Just remember
+that if you unsubscribe or switch GitHub user names, you'll need to re-fork and upload your changes from a local copy. 
\ No newline at end of file
