Update README.md

diff --git a/README.md b/README.md
--- a/README.md
+++ b/README.md
@@ -111,11 +111,11 @@ More info
 
 Visual Studio 2013 and Xcode 5.1 are strongly recommended for development.
 
-Legacy support for Visual Studio 2012 can be enabled by unzipping the [optional dependencies](https://github.com/EpicGames/UnrealEngine/releases/download/4.3.0-preview/Optional.zip) into your UnrealEngine folder
+Legacy support for Visual Studio 2012 can be enabled by unzipping the [optional dependencies](https://github.com/EpicGames/UnrealEngine/releases/download/latest-preview/Optional.zip) into your UnrealEngine folder
 and changing WindowsPlatform.Compiler to VisualStudio2012 in Engine\Source\Programs\UnrealBuildTool\Windows\UEBuildWindows.cs.
 
 The first time you start the editor from a fresh source build, you may experience long load times.  This only happens on the first 
 run as the engine optimizes content for the platform and _fills the derived data cache_.
 
 Your private forks of the Unreal Engine code are associated with your GitHub account permissions.  Just remember
-that if you unsubscribe or switch GitHub user names, you'll need to re-fork and upload your changes from a local copy. 
\ No newline at end of file
+that if you unsubscribe or switch GitHub user names, you'll need to re-fork and upload your changes from a local copy. 
