Update README.md

diff --git a/README.md b/README.md
--- a/README.md
+++ b/README.md
@@ -40,7 +40,7 @@ running the editor yourself on Windows or Mac.  Oh, and you might want to watch
 first.  Okay, here we go!
 
 1. [Download the source](https://github.com/EpicGames/UnrealEngine/archive/latest-preview.zip) and unzip it to a folder, or 
-   [create a fork](https://help.github.com/articles/fork-a-repo) and **clone the repository**.
+   [create a fork and clone the repository](https://help.github.com/articles/fork-a-repo).
 
 1. You should now have an _UnrealEngine_ folder on your computer.  All of the source and dependencies will go into this folder.  The folder name might 
    have a branch suffix (such as _UnrealEngine-master_), which is normal.
