Updated README.md

diff --git a/README.md b/README.md
--- a/README.md
+++ b/README.md
@@ -4,7 +4,7 @@ Welcome to the UE4 source code!
 This is the source code page for the **Unreal Engine on GitHub**.  With the UE4 source code, you can modify the
 engine and tools in any way imaginable and share your changes with others!
 
-You can build the editor for Windows and Mac and compile games for Android, iOS, Playstation 4, Xbox
+You can build the editor for Windows and Mac and compile games for Android, iOS, PlayStation 4, Xbox
 One and Linux.  Source code for all tools is included as well, such as Unreal Lightmass and Unreal Frontend.
 
 Before continuing, check out this [short tutorial video](http://youtu.be/usjlNHPn-jo) about getting started
@@ -72,7 +72,7 @@ first.  Okay, here we go!
 
 1. After compiling finishes, you can **load the editor** from Visual Studio by setting your startup project to **UE4** and pressing **F5** to debug.
 
-1. One last thing.  You'll want to setup your Windows shell so that you can interact with .uproject files.  Find the file named **UnrealVersionSelector-Win64-Shippping.exe** in 
+1. One last thing.  You'll want to setup your Windows shell so that you can interact with .uproject files.  Find the file named **UnrealVersionSelector-Win64-Shipping.exe** in 
    the _UnrealEngine/Engine/Binaries/Win64/_ folder and run it.  Now, you'll be able to double-click .uproject files to load the project, or right click them to quickly update Visual Studio files.         
 
 
@@ -102,7 +102,7 @@ first.  Okay, here we go!
 
 **Linux** support requires the [optional dependencies](https://github.com/EpicGames/UnrealEngine/releases/download/latest-preview/Optional.zip) to be unzipped into your _UnrealEngine_ folder, and the [cross-compiler toolchain](http://cdn.unrealengine.com/qfe/v3_clang-3.3_ld-2.24_glibc-2.12.2.zip) to be installed. See the [Linux cross-compiling page on the wiki](https://wiki.unrealengine.com/Compiling_For_Linux).
 
-**Playstation 4** or **XboxOne** development require additional files that can only be provided after your registered developer status is confirmed by Sony or Microsoft. See [the announcement blog post](https://www.unrealengine.com/blog/playstation-4-and-xbox-one-now-supported) for more information.
+**PlayStation 4** or **Xbox One** development require additional files that can only be provided after your registered developer status is confirmed by Sony or Microsoft. See [the announcement blog post](https://www.unrealengine.com/blog/playstation-4-and-xbox-one-now-supported) for more information.
 
 
 
