Updated .gitignore

diff --git a/.gitignore b/.gitignore
--- a/.gitignore
+++ b/.gitignore
@@ -47,6 +47,7 @@
 
 # Saved folders for programs should not be checked in
 /Engine/Programs/**/Saved/**
+/Engine/Programs/UnrealBuildTool/**
 
 # Ignore various other dependencies that are never checked in
 /Engine/Content/**
@@ -72,9 +73,9 @@
 /Engine/Source/ThirdParty/**/*.dylib
 /Engine/Source/ThirdParty/**/*.bmp
 /Engine/Source/ThirdParty/**/*.so
-/Engine/Source/ThirdParty/**/*.bc
 /Engine/Source/ThirdParty/**/*.bsc
 /Engine/Source/ThirdParty/**/*.idb
+/Engine/Source/ThirdParty/**/*.bc
 
 # Ignore sample and template projects
 /Samples/**
