Updated README with wiki link for getting started

diff --git a/README.md b/README.md
--- a/README.md
+++ b/README.md
@@ -7,7 +7,7 @@ engine and tools in any way imaginable and share your changes with others!
 You can build the editor for Windows and Mac and compile games for Android, iOS, PlayStation 4, Xbox
 One and Linux.  Source code for all tools is included as well, such as Unreal Lightmass and Unreal Frontend.
 
-Before continuing, check out this [short tutorial video](http://youtu.be/usjlNHPn-jo) about getting started
+Before continuing, check out this [wiki page](https://wiki.unrealengine.com/GitHub_Setup) about getting started
 with the engine code.  We also have a [programming guide](https://docs.unrealengine.com/latest/INT/Programming/index.html) and
 full [API documentation](https://docs.unrealengine.com/latest/INT/API/index.html).
 
