Updated README.md for automatic Visual Studio 2012 switching

diff --git a/README.md b/README.md
--- a/README.md
+++ b/README.md
@@ -58,7 +58,8 @@ first.  Okay, here we go!
 ### Windows
 
 1. Be sure to have [Visual Studio 2013](http://www.microsoft.com/en-us/download/details.aspx?id=40787) installed.  You can use any 
-   desktop version of Visual Studio 2013, including the free version:  [Visual Studio 2013 Express for Windows Desktop](http://www.microsoft.com/en-us/download/details.aspx?id=40787)
+   desktop version of Visual Studio 2013, including the free version:  [Visual Studio 2013 Express for Windows Desktop](http://www.microsoft.com/en-us/download/details.aspx?id=40787).
+   See the "Additional Notes" section below for using Visual Studio 2012.
 
 1. Make sure you have [June 2010 DirectX runtime](http://www.microsoft.com/en-us/download/details.aspx?id=8109) installed.  You don't need the SDK, just the runtime.
 
@@ -111,8 +112,7 @@ More info
 
 Visual Studio 2013 and Xcode 5.1 are strongly recommended for development.
 
-Legacy support for Visual Studio 2012 can be enabled by unzipping the [optional dependencies](https://github.com/EpicGames/UnrealEngine/releases/download/latest-preview/Optional.zip) into your UnrealEngine folder
-and changing WindowsPlatform.Compiler to VisualStudio2012 in Engine\Source\Programs\UnrealBuildTool\Windows\UEBuildWindows.cs.
+Legacy support for Visual Studio 2012 can be enabled by unzipping the [optional dependencies](https://github.com/EpicGames/UnrealEngine/releases/download/latest-preview/Optional.zip) into your UnrealEngine folder.
 
 The first time you start the editor from a fresh source build, you may experience long load times.  This only happens on the first 
 run as the engine optimizes content for the platform and _fills the derived data cache_.
