Updated README.md

diff --git a/README.md b/README.md
--- a/README.md
+++ b/README.md
@@ -93,6 +93,18 @@ first.  Okay, here we go!
 
 
 
+
+### Linux
+
+1. Both cross-compiling and native builds are supported. 
+
+2. Cross-compiling is handy when you are a Windows (Mac support planned too) developer who wants to package your game for Linux with minimal hassle, and it requires a [cross-compiler toolchain](http://cdn.unrealengine.com/qfe/v4_clang-3.5.0_ld-2.24_glibc-2.12.2.zip) to be installed (see the [Linux cross-compiling page on the wiki](https://wiki.unrealengine.com/Compiling_For_Linux)). Note that you will also need [optional dependencies](https://github.com/EpicGames/UnrealEngine/releases/download/latest-preview/Optional.zip) to be unzipped into your _UnrealEngine_ folder.
+
+2. Native compilation is discussed in [a separate README](https://github.com/EpicGames/UnrealEngine/blob/master/Engine/Build/BatchFiles/Linux/README.md) and [community wiki page](https://wiki.unrealengine.com/Building_On_Linux). Downloading the dependencies has now been automated, so you will only need to clone the repo and run **GenerateProjectFiles.sh** (provided that you have OAUTH_TOKEN set to your personal access token, see the above README for details).
+
+
+
+
 ### Additional target platforms
 
 **Android** development currently works best from a PC. See the [Android getting started guide](https://docs.unrealengine.com/latest/INT/Platforms/Android/GettingStarted/).
@@ -101,8 +113,6 @@ first.  Okay, here we go!
 
 **HTML5** is supported using Emscripten, and requires the [optional dependencies](https://github.com/EpicGames/UnrealEngine/releases/download/latest-preview/Optional.zip) to be unzipped into your _UnrealEngine_ folder. Please see the [HTML5 getting started guide](https://docs.unrealengine.com/latest/INT/Platforms/HTML5/GettingStarted/index.html).
 
-**Linux** support requires the [optional dependencies](https://github.com/EpicGames/UnrealEngine/releases/download/latest-preview/Optional.zip) to be unzipped into your _UnrealEngine_ folder, and the [cross-compiler toolchain](http://cdn.unrealengine.com/qfe/v3_clang-3.3_ld-2.24_glibc-2.12.2.zip) to be installed. See the [Linux cross-compiling page on the wiki](https://wiki.unrealengine.com/Compiling_For_Linux).
-
 **PlayStation 4** or **Xbox One** development require additional files that can only be provided after your registered developer status is confirmed by Sony or Microsoft. See [the announcement blog post](https://www.unrealengine.com/blog/playstation-4-and-xbox-one-now-supported) for more information.
 
 
