Updating .gitignore

diff --git a/.gitignore b/.gitignore
--- a/.gitignore
+++ b/.gitignore
@@ -91,3 +91,6 @@
 # Ignore Unix backup files
 *~
 
+# Ignore Mac desktop services store files 
+.DS_Store
+
