Updated README.md

diff --git a/README.md b/README.md
--- a/README.md
+++ b/README.md
@@ -46,8 +46,9 @@ first.  Okay, here we go!
    have a branch suffix (such as _UnrealEngine-master_), which is normal.
 
 1. Download the **required dependencies** files for the [latest preview tag](https://github.com/EpicGames/UnrealEngine/releases/tag/latest-preview): 
-   [Required_1of2.zip](https://github.com/EpicGames/UnrealEngine/releases/download/latest-preview/Required_1of2.zip), 
-   [Required_2of2.zip](https://github.com/EpicGames/UnrealEngine/releases/download/latest-preview/Required_2of2.zip).
+   [Required_1of3.zip](https://github.com/EpicGames/UnrealEngine/releases/download/latest-preview/Required_1of3.zip), 
+   [Required_2of3.zip](https://github.com/EpicGames/UnrealEngine/releases/download/latest-preview/Required_2of3.zip).
+   [Required_3of3.zip](https://github.com/EpicGames/UnrealEngine/releases/download/latest-preview/Required_3of3.zip).
 
 1. Unzip the dependencies into the _UnrealEngine_ folder alongside the source.  Be careful to make sure the folders are merged together 
    correctly.  On Mac, we recommend **Option + dragging** the unzipped files into the _UnrealEngine_ folder, then selecting **Keep Newer** if prompted.
