Updating .gitignore

diff --git a/.gitignore b/.gitignore
--- a/.gitignore
+++ b/.gitignore
@@ -84,6 +84,9 @@
 /Engine/Source/ThirdParty/**/*.idb
 /Engine/Source/ThirdParty/**/*.bc
 
+# Ignore LinuxNativeDialogs build files (all extensions)
+/Engine/Source/ThirdParty/LinuxNativeDialogs/UELinuxNativeDialogs/build/**
+
 # Ignore sample and template projects
 /Samples/**
 /Templates/**
@@ -94,3 +97,5 @@
 # Ignore Mac desktop services store files 
 .DS_Store
 
+# Ignore crash reports
+/crashinfo--*
