[Github PR 565] mpolaczyk - Snippets for Visual Studio

diff --git a/.gitignore b/.gitignore
--- a/.gitignore
+++ b/.gitignore
@@ -63,6 +63,7 @@
 /Engine/Extras/Maya_AnimationRiggingTools/**/*.pyd
 /Engine/Extras/Maya_AnimationRiggingTools/MayaTools/General/ART/**
 /Engine/Extras/Maya_AnimationRiggingTools/MayaTools/General/Icons/**
+!/Engine/Extras/VisualStudioSnippets/**
 
 # Ignore any saved local files
 /Engine/Saved/**
