Allow GitDependencies.exe to be stored in the repo.

diff --git a/.gitignore b/.gitignore
--- a/.gitignore
+++ b/.gitignore
@@ -42,6 +42,10 @@
 !/Engine/Binaries/ThirdParty/Mono/Mac/bin
 !/Engine/Binaries/ThirdParty/Mono/Mac/bin/**
 
+# ...and GitDependencies.exe
+!/Engine/Binaries/DotNET/GitDependencies.exe
+!/Engine/Binaries/DotNET/GitDependencies.exe.config
+
 # This is autogenerated symlink
 /Engine/Binaries/ThirdParty/Mono/Mac/bin/mono
 
