Snapshot for changelist 2348484.

[CL 2348484 in Main branch]

diff --git a/.gitignore b/.gitignore
--- a/.gitignore
+++ b/.gitignore
@@ -22,6 +22,7 @@
 /*.xcodeproj
 /Makefile
 /CMakeLists.txt
+/.ue4dependencies
 
 # Derived data cache is never checked in
 /Engine/DerivedDataCache/**
