Snapshot for changelist 2348932.

[CL 2348932 in Main branch]

diff --git a/GitSetup.bat b/GitSetup.bat
new file mode 100644
--- /dev/null
+++ b/GitSetup.bat
@@ -0,0 +1,26 @@
+@echo off
+setlocal
+pushd %~dp0
+
+rem Sync the dependencies...
+.\Engine\Binaries\DotNET\GitDependencies.exe %*
+if ERRORLEVEL 1 goto error
+
+rem Setup the git hooks
+if not exist .git\hooks goto no_git_hooks_directory
+echo Registering git hooks...
+echo #!/bin/sh >.git\hooks\post-checkout
+echo Engine/Binaries/DotNET/GitDependencies.exe >>.git\hooks\post-checkout
+:no_git_hooks_directory
+
+rem Register the engine installation...
+if not exist .\Engine\Binaries\Win64\UnrealVersionSelector-Win64-Shipping.exe goto :no_unreal_version_selector
+.\Engine\Binaries\Win64\UnrealVersionSelector-Win64-Shipping.exe /quiet
+:no_unreal_version_selector
+
+rem Done!
+goto :EOF
+
+rem Error happened. Wait for a keypress before quitting.
+:error
+pause
diff --git a/GitSetup.command b/GitSetup.command
new file mode 100755
--- /dev/null
+++ b/GitSetup.command
@@ -0,0 +1,7 @@
+#!/bin/sh
+# Copyright 1998-2014 Epic Games, Inc. All Rights Reserved.
+#
+# Simple wrapper around GitSetup.sh using the
+# .command extension enables it to be run from the OSX Finder.
+
+sh "`dirname "$0"`"/GitSetup.sh
diff --git a/GitSetup.sh b/GitSetup.sh
new file mode 100755
--- /dev/null
+++ b/GitSetup.sh
@@ -0,0 +1,20 @@
+#!/bin/sh
+# Copyright 1998-2014 Epic Games, Inc. All Rights Reserved.
+
+set -e
+
+cd "`dirname "$0"`"
+
+if [ ! -f Engine/Binaries/DotNET/GitDependencies.exe ]; then
+	echo "GitSetup ERROR: This script does not appear to be located \
+       in the root UE4 directory and must be run from there."
+	exit 1
+fi 
+
+if [ "$(uname)" = "Darwin" ]; then
+	cd Engine/Build/BatchFiles/Mac
+	source SetupMono.sh "`pwd`"
+	cd ../../../..
+fi
+
+mono Engine/Binaries/DotNET/GitDependencies.exe "$@"
