Comment out using special Mono distribution for Mac for now: doesn't work without having dependencies already.

diff --git a/GitSetup.sh b/GitSetup.sh
--- a/GitSetup.sh
+++ b/GitSetup.sh
@@ -11,10 +11,10 @@ if [ ! -f Engine/Binaries/DotNET/GitDependencies.exe ]; then
 	exit 1
 fi 
 
-if [ "$(uname)" = "Darwin" ]; then
-	cd Engine/Build/BatchFiles/Mac
-	source SetupMono.sh "`pwd`"
-	cd ../../../..
-fi
+#if [ "$(uname)" = "Darwin" ]; then
+#	cd Engine/Build/BatchFiles/Mac
+#	source SetupMono.sh "`pwd`"
+#	cd ../../../..
+#fi
 
 mono Engine/Binaries/DotNET/GitDependencies.exe "$@"
