Add a Git checkout hook which gets latest binaries.

diff --git a/GitSetup.sh b/GitSetup.sh
--- a/GitSetup.sh
+++ b/GitSetup.sh
@@ -11,6 +11,14 @@ if [ ! -f Engine/Binaries/DotNET/GitDependencies.exe ]; then
 	exit 1
 fi 
 
+# Setup the git hooks
+if [ -d .git/hooks ]; then
+	echo Registering git hooks...
+	echo \#!/bin/sh >.git/hooks/post-checkout
+	echo mono Engine/Binaries/DotNET/GitDependencies.exe >>.git/hooks/post-checkout
+	chmod +x .git/hooks/post-checkout
+fi
+
 #if [ "$(uname)" = "Darwin" ]; then
 #	cd Engine/Build/BatchFiles/Mac
 #	source SetupMono.sh "`pwd`"
