Add GitDependencies as a post-merge hook - missing file.

diff --git a/GitSetup.bat b/GitSetup.bat
--- a/GitSetup.bat
+++ b/GitSetup.bat
@@ -11,6 +11,8 @@ if not exist .git\hooks goto no_git_hooks_directory
 echo Registering git hooks...
 echo #!/bin/sh >.git\hooks\post-checkout
 echo Engine/Binaries/DotNET/GitDependencies.exe >>.git\hooks\post-checkout
+echo #!/bin/sh >.git\hooks\post-merge
+echo Engine/Binaries/DotNET/GitDependencies.exe >>.git\hooks\post-merge
 :no_git_hooks_directory
 
 rem Register the engine installation...
