Snapshot for changelist 2353950.

[CL 2353950 in Main branch]

diff --git a/GitSetup.bat b/GitSetup.bat
deleted file mode 100644
diff --git a/GitSetup.command b/GitSetup.command
deleted file mode 100755
diff --git a/GitSetup.sh b/GitSetup.sh
deleted file mode 100755
diff --git a/Setup.bat b/Setup.bat
new file mode 100644
--- /dev/null
+++ b/Setup.bat
@@ -0,0 +1,28 @@
+@echo off
+setlocal
+pushd %~dp0
+
+rem Sync the dependencies...
+.\Engine\Binaries\DotNET\GitDependencies.exe %*
+if ERRORLEVEL 1 goto error
+
+rem Setup the git hooks
+if not exist .git\hooks goto no_git_hooks_directory
+echo Registering git hooks...
+echo #!/bin/sh >.git\hooks\post-checkout
+echo Engine/Binaries/DotNET/GitDependencies.exe >>.git\hooks\post-checkout
+echo #!/bin/sh >.git\hooks\post-merge
+echo Engine/Binaries/DotNET/GitDependencies.exe >>.git\hooks\post-merge
+:no_git_hooks_directory
+
+rem Register the engine installation...
+if not exist .\Engine\Binaries\Win64\UnrealVersionSelector-Win64-Shipping.exe goto :no_unreal_version_selector
+.\Engine\Binaries\Win64\UnrealVersionSelector-Win64-Shipping.exe /register
+:no_unreal_version_selector
+
+rem Done!
+goto :EOF
+
+rem Error happened. Wait for a keypress before quitting.
+:error
+pause
diff --git a/Setup.command b/Setup.command
new file mode 100755
--- /dev/null
+++ b/Setup.command
@@ -0,0 +1,7 @@
+#!/bin/sh
+# Copyright 1998-2014 Epic Games, Inc. All Rights Reserved.
+#
+# Simple wrapper around Setup.sh using the
+# .command extension enables it to be run from the OSX Finder.
+
+sh "`dirname "$0"`"/Setup.sh
diff --git a/Setup.sh b/Setup.sh
new file mode 100755
--- /dev/null
+++ b/Setup.sh
@@ -0,0 +1,32 @@
+#!/bin/sh
+# Copyright 1998-2014 Epic Games, Inc. All Rights Reserved.
+
+set -e
+
+cd "`dirname "$0"`"
+
+if [ ! -f Engine/Binaries/DotNET/GitDependencies.exe ]; then
+	echo "GitSetup ERROR: This script does not appear to be located \
+       in the root UE4 directory and must be run from there."
+	exit 1
+fi 
+
+# Setup the git hooks
+if [ -d .git/hooks ]; then
+	echo "Registering git hooks... (this will override existing ones!)"
+	echo \#!/bin/sh >.git/hooks/post-checkout
+	echo mono Engine/Binaries/DotNET/GitDependencies.exe >>.git/hooks/post-checkout
+	chmod +x .git/hooks/post-checkout
+
+	echo \#!/bin/sh >.git/hooks/post-merge
+	echo mono Engine/Binaries/DotNET/GitDependencies.exe >>.git/hooks/post-merge
+	chmod +x .git/hooks/post-merge
+fi
+
+#if [ "$(uname)" = "Darwin" ]; then
+#	cd Engine/Build/BatchFiles/Mac
+#	source SetupMono.sh "`pwd`"
+#	cd ../../../..
+#fi
+
+mono Engine/Binaries/DotNET/GitDependencies.exe "$@"
