Install prerequisites as part of Git setup

diff --git a/Setup.bat b/Setup.bat
--- a/Setup.bat
+++ b/Setup.bat
@@ -6,7 +6,7 @@ rem Sync the dependencies...
 .\Engine\Binaries\DotNET\GitDependencies.exe %*
 if ERRORLEVEL 1 goto error
 
-rem Setup the git hooks
+rem Setup the git hooks...
 if not exist .git\hooks goto no_git_hooks_directory
 echo Registering git hooks...
 echo #!/bin/sh >.git\hooks\post-checkout
@@ -15,6 +15,10 @@ echo #!/bin/sh >.git\hooks\post-merge
 echo Engine/Binaries/DotNET/GitDependencies.exe >>.git\hooks\post-merge
 :no_git_hooks_directory
 
+rem Install prerequisites...
+echo Installing prerequisites...
+start /wait Engine\Extras\Redist\en-us\UE4PrereqSetup_x64.exe /quiet
+
 rem Register the engine installation...
 if not exist .\Engine\Binaries\Win64\UnrealVersionSelector-Win64-Shipping.exe goto :no_unreal_version_selector
 .\Engine\Binaries\Win64\UnrealVersionSelector-Win64-Shipping.exe /register
