Switch back to using bundled mono on Mac. Store a minimal install in the repository.

diff --git a/Setup.sh b/Setup.sh
--- a/Setup.sh
+++ b/Setup.sh
@@ -23,10 +23,10 @@ if [ -d .git/hooks ]; then
 	chmod +x .git/hooks/post-merge
 fi
 
-#if [ "$(uname)" = "Darwin" ]; then
-#	cd Engine/Build/BatchFiles/Mac
-#	source SetupMono.sh "`pwd`"
-#	cd ../../../..
-#fi
+if [ "$(uname)" = "Darwin" ]; then
+	cd Engine/Build/BatchFiles/Mac
+	source SetupMono.sh "`pwd`"
+	cd ../../../..
+fi
 
 mono Engine/Binaries/DotNET/GitDependencies.exe "$@"
