Change git hook to a symlink to Engine/Build/BatchFiles/Mac/GitDependenciesHook.sh (which calls GitDependencies.sh), and allow GitDependencies.sh to be run as a link from elsewhere.

diff --git a/Engine/Build/BatchFiles/Mac/GitDependencies.sh b/Engine/Build/BatchFiles/Mac/GitDependencies.sh
--- a/Engine/Build/BatchFiles/Mac/GitDependencies.sh
+++ b/Engine/Build/BatchFiles/Mac/GitDependencies.sh
@@ -3,15 +3,18 @@
 
 set -e
 
-cd "`dirname "$0"`"
+SCRIPT_PATH=$0
+if [ -h $SCRIPT_PATH ]; then
+	SCRIPT_PATH=$(dirname $SCRIPT_PATH)/$(readlink $SCRIPT_PATH)
+fi
+
+cd $(dirname $SCRIPT_PATH)
 
 if [ ! -f ../../../Binaries/DotNET/GitDependencies.exe ]; then
 	echo "Cannot find GitDependencies.exe. This script should be placed in Engine/Build/BatchFiles/Mac."
 	exit 1
 fi 
 
-echo from `pwd`
-
 source SetupMono.sh "`pwd`"
 
 cd ../../../..
diff --git a/Engine/Build/BatchFiles/Mac/GitDependenciesHook.sh b/Engine/Build/BatchFiles/Mac/GitDependenciesHook.sh
new file mode 100755
--- /dev/null
+++ b/Engine/Build/BatchFiles/Mac/GitDependenciesHook.sh
@@ -0,0 +1,13 @@
+#!/bin/sh
+# Copyright 1998-2014 Epic Games, Inc. All Rights Reserved.
+
+set -e
+
+SCRIPT_PATH=$0
+if [ -h $SCRIPT_PATH ]; then
+	SCRIPT_PATH=$(dirname $SCRIPT_PATH)/$(readlink $SCRIPT_PATH)
+fi
+cd $(dirname $SCRIPT_PATH)
+
+# Swallow the arguments passed in to this script; we don't need them when running a hook
+./GitDependencies.sh
diff --git a/Setup.sh b/Setup.sh
--- a/Setup.sh
+++ b/Setup.sh
@@ -11,22 +11,30 @@ if [ ! -f Engine/Binaries/DotNET/GitDependencies.exe ]; then
 	exit 1
 fi 
 
-# Setup the git hooks
-if [ -d .git/hooks ]; then
-	echo "Registering git hooks... (this will override existing ones!)"
-	echo \#!/bin/sh >.git/hooks/post-checkout
-	echo mono Engine/Binaries/DotNET/GitDependencies.exe >>.git/hooks/post-checkout
-	chmod +x .git/hooks/post-checkout
+if [ "$(uname)" = "Darwin" ]; then
+	# Setup the git hooks
+	if [ -d .git/hooks ]; then
+		echo "Registering git hooks... (this will override existing ones!)"
+		rm -f .git/hooks/post-checkout
+		rm -f .git/hooks/post-merge
+		ln -s ../../Engine/Build/BatchFiles/Mac/GitDependenciesHook.sh .git/hooks/post-checkout
+		ln -s ../../Engine/Build/BatchFiles/Mac/GitDependenciesHook.sh .git/hooks/post-merge
+	fi
 
-	echo \#!/bin/sh >.git/hooks/post-merge
-	echo mono Engine/Binaries/DotNET/GitDependencies.exe >>.git/hooks/post-merge
-	chmod +x .git/hooks/post-merge
-fi
+	# Get the dependencies for the first time
+	Engine/Build/BatchFiles/Mac/GitDependencies.sh $@
+else
+	# Setup the git hooks
+	if [ -d .git/hooks ]; then
+		echo "Registering git hooks... (this will override existing ones!)"
+		echo \#!/bin/sh >.git/hooks/post-checkout
+		echo mono Engine/Binaries/DotNET/GitDependencies.exe >>.git/hooks/post-checkout
+		chmod +x .git/hooks/post-checkout
 
-if [ "$(uname)" = "Darwin" ]; then
-	cd Engine/Build/BatchFiles/Mac
-	source SetupMono.sh "`pwd`"
-	cd ../../../..
-fi
+		echo \#!/bin/sh >.git/hooks/post-merge
+		echo mono Engine/Binaries/DotNET/GitDependencies.exe >>.git/hooks/post-merge
+		chmod +x .git/hooks/post-merge
+	fi
 
-mono Engine/Binaries/DotNET/GitDependencies.exe "$@"
+	mono Engine/Binaries/DotNET/GitDependencies.exe "$@"
+fi
