Snapshot for changelist 2358805.

[CL 2358805 in Main branch]

diff --git a/.gitignore b/.gitignore
--- a/.gitignore
+++ b/.gitignore
@@ -37,44 +37,44 @@
 
 # Add exceptions for Mac Mono binaries needed to run GitDependencies.exe. This is a mess, but we need to un-ignore every parent directory before anything in it,
 # and Mono has a pretty deep hierarchy.
-!/Engine/Binaries
 !/Engine/Binaries/ThirdParty
+/Engine/Binaries/ThirdParty/**
 !/Engine/Binaries/ThirdParty/Mono
+/Engine/Binaries/ThirdParty/Mono/**
 !/Engine/Binaries/ThirdParty/Mono/Mac
+/Engine/Binaries/ThirdParty/Mono/Mac/**
 !/Engine/Binaries/ThirdParty/Mono/Mac/bin
-!/Engine/Binaries/ThirdParty/Mono/Mac/bin/*
+
 !/Engine/Binaries/ThirdParty/Mono/Mac/etc
 !/Engine/Binaries/ThirdParty/Mono/Mac/etc/mono
+/Engine/Binaries/ThirdParty/Mono/Mac/etc/mono/**
 !/Engine/Binaries/ThirdParty/Mono/Mac/etc/mono/config
 !/Engine/Binaries/ThirdParty/Mono/Mac/etc/mono/4.5
+/Engine/Binaries/ThirdParty/Mono/Mac/etc/mono/4.5/**
 !/Engine/Binaries/ThirdParty/Mono/Mac/etc/mono/4.5/machine.config
+
 !/Engine/Binaries/ThirdParty/Mono/Mac/lib/
+/Engine/Binaries/ThirdParty/Mono/Mac/lib/**
 !/Engine/Binaries/ThirdParty/Mono/Mac/lib/libMonoPosixHelper.dylib
 !/Engine/Binaries/ThirdParty/Mono/Mac/lib/mono
+/Engine/Binaries/ThirdParty/Mono/Mac/lib/mono/**
 !/Engine/Binaries/ThirdParty/Mono/Mac/lib/mono/4.5
+/Engine/Binaries/ThirdParty/Mono/Mac/lib/mono/4.5/**
 !/Engine/Binaries/ThirdParty/Mono/Mac/lib/mono/4.5/mscorlib.dll
+
 !/Engine/Binaries/ThirdParty/Mono/Mac/lib/mono/gac
+/Engine/Binaries/ThirdParty/Mono/Mac/lib/mono/gac/**
 !/Engine/Binaries/ThirdParty/Mono/Mac/lib/mono/gac/Mono.Posix
-!/Engine/Binaries/ThirdParty/Mono/Mac/lib/mono/gac/Mono.Posix/4.0.0.0__0738eb9f132ed756
-!/Engine/Binaries/ThirdParty/Mono/Mac/lib/mono/gac/Mono.Posix/4.0.0.0__0738eb9f132ed756/Mono.Posix.dll
 !/Engine/Binaries/ThirdParty/Mono/Mac/lib/mono/gac/Mono.Security
-!/Engine/Binaries/ThirdParty/Mono/Mac/lib/mono/gac/Mono.Security/4.0.0.0__0738eb9f132ed756
-!/Engine/Binaries/ThirdParty/Mono/Mac/lib/mono/gac/Mono.Security/4.0.0.0__0738eb9f132ed756/Mono.Security.dll
 !/Engine/Binaries/ThirdParty/Mono/Mac/lib/mono/gac/System
-!/Engine/Binaries/ThirdParty/Mono/Mac/lib/mono/gac/System/4.0.0.0__b77a5c561934e089
-!/Engine/Binaries/ThirdParty/Mono/Mac/lib/mono/gac/System/4.0.0.0__b77a5c561934e089/System.dll
 !/Engine/Binaries/ThirdParty/Mono/Mac/lib/mono/gac/System.Configuration
-!/Engine/Binaries/ThirdParty/Mono/Mac/lib/mono/gac/System.Configuration/4.0.0.0__b03f5f7f11d50a3a
-!/Engine/Binaries/ThirdParty/Mono/Mac/lib/mono/gac/System.Configuration/4.0.0.0__b03f5f7f11d50a3a/System.Configuration.dll
 !/Engine/Binaries/ThirdParty/Mono/Mac/lib/mono/gac/System.Core
-!/Engine/Binaries/ThirdParty/Mono/Mac/lib/mono/gac/System.Core/4.0.0.0__b77a5c561934e089
-!/Engine/Binaries/ThirdParty/Mono/Mac/lib/mono/gac/System.Core/4.0.0.0__b77a5c561934e089/System.Core.dll
 !/Engine/Binaries/ThirdParty/Mono/Mac/lib/mono/gac/System.Xml
-!/Engine/Binaries/ThirdParty/Mono/Mac/lib/mono/gac/System.Xml/4.0.0.0__b77a5c561934e089
-!/Engine/Binaries/ThirdParty/Mono/Mac/lib/mono/gac/System.Xml/4.0.0.0__b77a5c561934e089/System.Xml.dll
+/Engine/Binaries/ThirdParty/Mono/Mac/lib/mono/gac/**/*.mdb
 
 # ...and GitDependencies.exe
 !/Engine/Binaries/DotNET
+/Engine/Binaries/DotNET/**
 !/Engine/Binaries/DotNET/GitDependencies.exe
 !/Engine/Binaries/DotNET/GitDependencies.exe.config
 
