Update README.md for master.

diff --git a/README.md b/README.md
--- a/README.md
+++ b/README.md
@@ -1,87 +1,75 @@
-Welcome to the UE4 source code!
-===============================
+Unreal Engine
+=============
 
-This is the source code page for the **Unreal Engine on GitHub**.  With the UE4 source code, you can modify the
-engine and tools in any way imaginable and share your changes with others!
+Welcome to the Unreal Engine source code! 
 
-You can build the editor for Windows and Mac and compile games for Android, iOS, PlayStation 4, Xbox
-One and Linux.  Source code for all tools is included as well, such as Unreal Lightmass and Unreal Frontend.
+From this repository you can build the Unreal Editor for Windows and Mac, compile Unreal Engine games for Android, iOS, Playstation 4, Xbox One, HTML5 and Linux,
+and build tools like Unreal Lightmass and Unreal Frontend. Modify them in any way you can imagine, and share your changes with others! 
 
-Before continuing, check out this [wiki page](https://wiki.unrealengine.com/GitHub_Setup) about getting started
-with the engine code.  We also have a [programming guide](https://docs.unrealengine.com/latest/INT/Programming/index.html) and
-full [API documentation](https://docs.unrealengine.com/latest/INT/API/index.html).
+We have a heap of documentation available for the engine on the web. If you're looking for the answer to something, you may want to start here: 
 
-We also have a forum section where you can discuss [engine source and GitHub](https://forums.unrealengine.com/forumdisplay.php?1-Development-Discussion).
-Have fun!!  We can't wait to see what you create!!
+* [Unreal Engine Programming Guide](https://docs.unrealengine.com/latest/INT/Programming/index.html)
+* [Unreal Engine API Reference](https://docs.unrealengine.com/latest/INT/API/index.html)
+* [Engine source and GitHub on the Unreal Engine forums](https://forums.unrealengine.com/forumdisplay.php?1-Development-Discussion)
 
+If you need more, just ask. A lot of Epic developers hang out on the [forums](https://forums.unrealengine.com/) or [AnswerHub](https://answers.unrealengine.com/), 
+and we're proud to be part of a well-meaning, friendly and welcoming community of thousands. 
 
 
-Source releases
----------------
+Branches
+--------
 
-This branch mirrors the live codebase that we're working on at Epic, and contains [bleeding-edge changes](https://github.com/EpicGames/UnrealEngine/commits/master) 
-by our engine team. The head revision may be unstable or not even compile, so we encourage working from a versioned release branch which is fully 
-tested by our QA team instead. Official releases can be found [here](https://github.com/EpicGames/UnrealEngine/releases).
+We publish source for the engine in three rolling branches:
 
-You can **download the source** in .zip format, but in order to *contribute to the community* you should
-**setup a Git client** on your computer.  [This page](http://help.github.com/articles/set-up-git) will have 
-the info you need.
+The **[release branches](https://github.com/EpicGames/UnrealEngine/tree/4.6)** are extensively tested by our QA team and makes a great starting point for learning the engine or
+making your own games. We work hard to make releases stable and reliable, and aim to publish new releases every 1-2 months.
 
-Remember, you'll need to _download dependencies_ in order to actually build and run the engine.  We periodically update the latest-preview tag with matching 
-dependencies on our [releases page](https://github.com/EpicGames/UnrealEngine/releases).
+The **[promoted branch](https://github.com/EpicGames/UnrealEngine/tree/promoted)** branch is updated with builds for our artists and designers to use. We try to update it daily, though we often
+catch things that prevent us from doing so, and it's a good balance between getting the latest cool stuff and knowing most things work.
 
+The **[master branch](https://github.com/EpicGames/UnrealEngine/tree/master)** branch tracks [live changes](https://github.com/EpicGames/UnrealEngine/commits/master) by our engine team. This is the cutting edge and may be buggy - it may not even 
+compile. Battle-hardened developers eager to work lock-step with us on the latest and greatest should head here.
 
 
 Getting up and running
 ----------------------
 
-Here is the fun part!  This is a quick start guide to getting up and running with the source.  The steps below will take you through cloning your own private fork, then compiling and 
-running the editor yourself on Windows or Mac.  Oh, and you might want to watch our [short tutorial video](http://youtu.be/usjlNHPn-jo)
-first.  Okay, here we go!
-
-1. [Download the source](https://github.com/EpicGames/UnrealEngine/archive/latest-preview.zip) and unzip it to a folder, or 
-   [create a fork and clone the repository](https://help.github.com/articles/fork-a-repo).
-
-1. You should now have an _UnrealEngine_ folder on your computer.  All of the source and dependencies will go into this folder.  The folder name might 
-   have a branch suffix (such as _UnrealEngine-master_), which is normal.
-
-1. Download the **required dependencies** files for the [latest preview tag](https://github.com/EpicGames/UnrealEngine/releases/tag/latest-preview): 
-   [Required_1of3.zip](https://github.com/EpicGames/UnrealEngine/releases/download/latest-preview/Required_1of3.zip),
-   [Required_2of3.zip](https://github.com/EpicGames/UnrealEngine/releases/download/latest-preview/Required_2of3.zip),
-   [Required_3of3.zip](https://github.com/EpicGames/UnrealEngine/releases/download/latest-preview/Required_3of3.zip).
-
-1. Unzip the dependencies into the _UnrealEngine_ folder alongside the source.  Be careful to make sure the folders are merged together 
-   correctly.  On Mac, we recommend **Option + dragging** the unzipped files into the _UnrealEngine_ folder, then selecting **Keep Newer** if prompted.
-
-1. Okay, platform stuff comes next.  Depending on whether you are on Windows or Mac, follow one of the sections below:
-
+The steps below will take you through cloning your own private fork, then compiling and running the editor yourself:
 
 ### Windows
 
-1. Be sure to have [Visual Studio 2013](http://www.microsoft.com/en-us/download/details.aspx?id=40787) installed.  You can use any 
-   desktop version of Visual Studio 2013, including the free version:  [Visual Studio 2013 Express for Windows Desktop](http://www.microsoft.com/en-us/download/details.aspx?id=40787).
-   See the "Additional Notes" section below for using Visual Studio 2012.
-
-1. Make sure you have [June 2010 DirectX runtime](http://www.microsoft.com/en-us/download/details.aspx?id=8109) installed.  You don't need the SDK, just the runtime.
+1. Install **[GitHub for Windows](https://windows.github.com/)** then **fork our repository** and **clone it to your computer**. 
+   If you'd rather use Git from the command line, see the [Setting up Git](https://help.github.com/articles/set-up-git/) and [Fork a Repo](https://help.github.com/articles/fork-a-repo/) articles.
+   If you'd prefer not to use Git at all, you can [download the source as a zip file](https://github.com/EpicGames/UnrealEngine/archive/master.zip).
 
-1. You'll need project files in order to compile.  In the _UnrealEngine_ folder, double-click on **GenerateProjectFiles.bat**.  It should take less than a minute to complete.  On Windows 8, a warning from SmartScreen may appear.  Click "More info", then "Run anyway" to continue.
+1. Make sure you have **Visual Studio 2013** installed.
+   Any desktop edition of Visual Studio 2013 will do, including the free version: [Visual Studio 2013 Express for Windows Desktop](http://www.microsoft.com/en-us/download/details.aspx?id=40787).
 
-1. Load the project into Visual Studio by double-clicking on the **UE4.sln** file.
+1. Open your source folder in Explorer and run **Setup.bat**. This will download binary content for the engine, as well as installing prerequisites
+   and setting up Unreal file associations. On Windows 8, a warning from SmartScreen may appear.  Click "More info", then "Run anyway" to continue.
+   
+   A clean download of the engine binaries is currently around 2.5gb, which may take some time to complete on a slow internet connection.
+   Subsequent checkouts only require incremental downloads and will be much quicker.
+ 
+1. Run **GenerateProjectFiles.bat** to create project files for the engine. It should take less than a minute to complete.  
 
-1. It's time to **compile the editor**!  In Visual Studio, make sure your solution configuration is set to **Development Editor**, and your solution 
-   platform is set to **Win64**.  Right click on the **UE4** target and select **Build**.  It will take between 15 and 40 minutes to finish compiling,
-   depending on your system specs.
+1. Load the project into Visual Studio by double-clicking on the **UE4.sln** file. Set your solution configuration to **Development Editor** and your solution
+   platform to **Win64**, then right click on the **UE4** target and select **Build**. It may take anywhere between 10 and 40 minutes to finish compiling, depending on your system specs.
 
-1. After compiling finishes, you can **load the editor** from Visual Studio by setting your startup project to **UE4** and pressing **F5** to debug.
+1. After compiling finishes, you can load the editor from Visual Studio by setting your startup project to **UE4** and pressing **F5** to debug.
 
-1. One last thing.  You'll want to setup your Windows shell so that you can interact with .uproject files.  Find the file named **UnrealVersionSelector-Win64-Shipping.exe** in 
-   the _UnrealEngine/Engine/Binaries/Win64/_ folder and run it.  Now, you'll be able to double-click .uproject files to load the project, or right click them to quickly update Visual Studio files.         
 
 
 
 ### Mac
+   
+1. Install **[GitHub for Mac](https://mac.github.com/)** then **fork our repository** and **clone it to your computer**. 
+   To use Git from the Terminal, see the [Setting up Git](https://help.github.com/articles/set-up-git/) and [Fork a Repo](https://help.github.com/articles/fork-a-repo/) articles. 
+   Or if you'd rather not use Git at all, you can [download the source as a zip file](https://github.com/EpicGames/UnrealEngine/archive/master.zip).
+   
+1. Install the latest version of [Xcode](https://itunes.apple.com/us/app/xcode/id497799835).
 
-1. Be sure to have [Xcode 5.1](https://itunes.apple.com/us/app/xcode/id497799835) installed.
+1. Open your source folder in Finder and run **Setup.command** to download binary content for the engine.
 
 1. You'll need project files in order to compile.  In the _UnrealEngine_ folder, double-click on **GenerateProjectFiles.command**.  It should take less than a minute to complete.  You can close the Terminal window afterwards.  If you downloaded the source in .zip format, you may see a warning about an unidentified developer.  This is because because the .zip files on GitHub are not digitally signed.  To work around this, right-click on GenerateProjectFiles.command, select Open, then click the Open button if you are sure you want to open it.
 
@@ -97,36 +85,39 @@ first.  Okay, here we go!
 
 ### Linux
 
+1. [Set up Git](https://help.github.com/articles/set-up-git/) and [fork our repository](https://help.github.com/articles/fork-a-repo/).
+   If you'd prefer not to use Git, you can [download the source as a zip file](https://github.com/EpicGames/UnrealEngine/archive/master.zip).
+
+1. Open your source folder and run **Setup.sh** to download binary content for the engine.
+
 1. Both cross-compiling and native builds are supported. 
 
-2. Cross-compiling is handy when you are a Windows (Mac support planned too) developer who wants to package your game for Linux with minimal hassle, and it requires a [cross-compiler toolchain](http://cdn.unrealengine.com/qfe/v4_clang-3.5.0_ld-2.24_glibc-2.12.2.zip) to be installed (see the [Linux cross-compiling page on the wiki](https://wiki.unrealengine.com/Compiling_For_Linux)). Note that you will also need [optional dependencies](https://github.com/EpicGames/UnrealEngine/releases/download/latest-preview/Optional.zip) to be unzipped into your _UnrealEngine_ folder.
+   **Cross-compiling** is handy when you are a Windows (Mac support planned too) developer who wants to package your game for Linux with minimal hassle, and it requires a [cross-compiler toolchain](http://cdn.unrealengine.com/qfe/v4_clang-3.5.0_ld-2.24_glibc-2.12.2.zip) to be installed (see the [Linux cross-compiling page on the wiki](https://wiki.unrealengine.com/Compiling_For_Linux)).
 
-2. Native compilation is discussed in [a separate README](https://github.com/EpicGames/UnrealEngine/blob/master/Engine/Build/BatchFiles/Linux/README.md) and [community wiki page](https://wiki.unrealengine.com/Building_On_Linux). Downloading the dependencies has now been automated, so you will only need to clone the repo and run **GenerateProjectFiles.sh** (provided that you have OAUTH_TOKEN set to your personal access token, see the above README for details).
+   **Native compilation** is discussed in [a separate README](https://github.com/EpicGames/UnrealEngine/blob/master/Engine/Build/BatchFiles/Linux/README.md) and [community wiki page](https://wiki.unrealengine.com/Building_On_Linux). Downloading the dependencies has now been automated, so you will only need to clone the repo and run **GenerateProjectFiles.sh** (provided that you have OAUTH_TOKEN set to your personal access token, see the above README for details).
 
 
 
 
 ### Additional target platforms
 
-**Android** development currently works best from a PC. See the [Android getting started guide](https://docs.unrealengine.com/latest/INT/Platforms/Android/GettingStarted/).
+**Android** support will be downloaded by the setup script if you have the Android NDK installed. See the [Android getting started guide](https://docs.unrealengine.com/latest/INT/Platforms/Android/GettingStarted/).
 
 **iOS** programming requires a Mac. Instructions are in the [iOS getting started guide](https://docs.unrealengine.com/latest/INT/Platforms/iOS/GettingStarted/index.html).
 
-**HTML5** is supported using Emscripten, and requires the [optional dependencies](https://github.com/EpicGames/UnrealEngine/releases/download/latest-preview/Optional.zip) to be unzipped into your _UnrealEngine_ folder. Please see the [HTML5 getting started guide](https://docs.unrealengine.com/latest/INT/Platforms/HTML5/GettingStarted/index.html).
-
-**PlayStation 4** or **Xbox One** development require additional files that can only be provided after your registered developer status is confirmed by Sony or Microsoft. See [the announcement blog post](https://www.unrealengine.com/blog/playstation-4-and-xbox-one-now-supported) for more information.
+**HTML5** support will be downloaded by the setup script if you have Emscripten installed. Please see the [HTML5 getting started guide](https://docs.unrealengine.com/latest/INT/Platforms/HTML5/GettingStarted/index.html).
 
+**Playstation 4** or **XboxOne** development require additional files that can only be provided after your registered developer status is confirmed by Sony or Microsoft. See [the announcement blog post](https://www.unrealengine.com/blog/playstation-4-and-xbox-one-now-supported) for more information.
 
 
-More info
----------
 
-Visual Studio 2013 and Xcode 5.1 are strongly recommended for development.
+Additional Notes
+----------------
 
-Legacy support for Visual Studio 2012 can be enabled by unzipping the [optional dependencies](https://github.com/EpicGames/UnrealEngine/releases/download/latest-preview/Optional.zip) into your UnrealEngine folder.
+Visual Studio 2012 is supported by the Windows toolchain, though Visual Studio 2013 is recommended.
 
-The first time you start the editor from a fresh source build, you may experience long load times.  This only happens on the first 
-run as the engine optimizes content for the platform and _fills the derived data cache_.
+The first time you start the editor from a fresh source build, you may experience long load times. 
+The engine is optimizing content for your platform to the _derived data cache_, and it should only happen once.
 
-Your private forks of the Unreal Engine code are associated with your GitHub account permissions.  Just remember
-that if you unsubscribe or switch GitHub user names, you'll need to re-fork and upload your changes from a local copy. 
+Your private forks of the Unreal Engine code are associated with your GitHub account permissions.
+If you unsubscribe or switch GitHub user names, you'll need to re-fork and upload your changes from a local copy. 
