Updated README.md

diff --git a/README.md b/README.md
--- a/README.md
+++ b/README.md
@@ -21,7 +21,7 @@ Branches
 
 We publish source for the engine in three rolling branches:
 
-The **[release branches](https://github.com/EpicGames/UnrealEngine/tree/4.6)** are extensively tested by our QA team and makes a great starting point for learning the engine or
+The **[release branch](https://github.com/EpicGames/UnrealEngine/tree/release)** is extensively tested by our QA team and makes a great starting point for learning the engine or
 making your own games. We work hard to make releases stable and reliable, and aim to publish new releases every 1-2 months.
 
 The **[promoted branch](https://github.com/EpicGames/UnrealEngine/tree/promoted)** branch is updated with builds for our artists and designers to use. We try to update it daily, though we often
@@ -40,7 +40,7 @@ The steps below will take you through cloning your own private fork, then compil
 
 1. Install **[GitHub for Windows](https://windows.github.com/)** then **fork our repository** and **clone it to your computer**. 
    If you'd rather use Git from the command line, see the [Setting up Git](https://help.github.com/articles/set-up-git/) and [Fork a Repo](https://help.github.com/articles/fork-a-repo/) articles.
-   If you'd prefer not to use Git at all, you can [download the source as a zip file](https://github.com/EpicGames/UnrealEngine/archive/master.zip).
+   If you'd prefer not to use Git at all, download the source as a zip file with the button on the right.
 
 1. Make sure you have **Visual Studio 2013** installed.
    Any desktop edition of Visual Studio 2013 will do, including the free version: [Visual Studio 2013 Express for Windows Desktop](http://www.microsoft.com/en-us/download/details.aspx?id=40787).
@@ -65,8 +65,8 @@ The steps below will take you through cloning your own private fork, then compil
    
 1. Install **[GitHub for Mac](https://mac.github.com/)** then **fork our repository** and **clone it to your computer**. 
    To use Git from the Terminal, see the [Setting up Git](https://help.github.com/articles/set-up-git/) and [Fork a Repo](https://help.github.com/articles/fork-a-repo/) articles. 
-   Or if you'd rather not use Git at all, you can [download the source as a zip file](https://github.com/EpicGames/UnrealEngine/archive/master.zip).
-   
+   Or if you'd rather not use Git at all, download the source as a zip file with the button on the right.
+
 1. Install the latest version of [Xcode](https://itunes.apple.com/us/app/xcode/id497799835).
 
 1. Open your source folder in Finder and run **Setup.command** to download binary content for the engine.
@@ -86,7 +86,7 @@ The steps below will take you through cloning your own private fork, then compil
 ### Linux
 
 1. [Set up Git](https://help.github.com/articles/set-up-git/) and [fork our repository](https://help.github.com/articles/fork-a-repo/).
-   If you'd prefer not to use Git, you can [download the source as a zip file](https://github.com/EpicGames/UnrealEngine/archive/master.zip).
+   If you'd prefer not to use Git, download the source as a zip file with the button on the right.
 
 1. Open your source folder and run **Setup.sh** to download binary content for the engine.
 
@@ -94,7 +94,7 @@ The steps below will take you through cloning your own private fork, then compil
 
    **Cross-compiling** is handy when you are a Windows (Mac support planned too) developer who wants to package your game for Linux with minimal hassle, and it requires a [cross-compiler toolchain](http://cdn.unrealengine.com/qfe/v4_clang-3.5.0_ld-2.24_glibc-2.12.2.zip) to be installed (see the [Linux cross-compiling page on the wiki](https://wiki.unrealengine.com/Compiling_For_Linux)).
 
-   **Native compilation** is discussed in [a separate README](https://github.com/EpicGames/UnrealEngine/blob/master/Engine/Build/BatchFiles/Linux/README.md) and [community wiki page](https://wiki.unrealengine.com/Building_On_Linux). Downloading the dependencies has now been automated, so you will only need to clone the repo and run **GenerateProjectFiles.sh** (provided that you have OAUTH_TOKEN set to your personal access token, see the above README for details).
+   **Native compilation** is discussed in [a separate README](Engine/Build/BatchFiles/Linux/README.md) and [community wiki page](https://wiki.unrealengine.com/Building_On_Linux). Downloading the dependencies has now been automated, so you will only need to clone the repo and run **GenerateProjectFiles.sh** (provided that you have OAUTH_TOKEN set to your personal access token, see the above README for details).
 
 
 
