Updated README.md again

diff --git a/README.md b/README.md
--- a/README.md
+++ b/README.md
@@ -40,7 +40,7 @@ The steps below will take you through cloning your own private fork, then compil
 
 1. Install **[GitHub for Windows](https://windows.github.com/)** then **fork our repository** and **clone it to your computer**. 
    If you'd rather use Git from the command line, see the [Setting up Git](https://help.github.com/articles/set-up-git/) and [Fork a Repo](https://help.github.com/articles/fork-a-repo/) articles.
-   If you'd prefer not to use Git at all, download the source as a zip file with the button on the right.
+   If you'd prefer not to use Git at all, use the 'Download ZIP' button on the right to get the source as a zip file.
 
 1. Make sure you have **Visual Studio 2013** installed.
    Any desktop edition of Visual Studio 2013 will do, including the free version: [Visual Studio 2013 Express for Windows Desktop](http://www.microsoft.com/en-us/download/details.aspx?id=40787).
@@ -65,7 +65,7 @@ The steps below will take you through cloning your own private fork, then compil
    
 1. Install **[GitHub for Mac](https://mac.github.com/)** then **fork our repository** and **clone it to your computer**. 
    To use Git from the Terminal, see the [Setting up Git](https://help.github.com/articles/set-up-git/) and [Fork a Repo](https://help.github.com/articles/fork-a-repo/) articles. 
-   Or if you'd rather not use Git at all, download the source as a zip file with the button on the right.
+   Or if you'd rather not use Git at all, use the 'Download ZIP' button on the right to get the source as a zip file.
 
 1. Install the latest version of [Xcode](https://itunes.apple.com/us/app/xcode/id497799835).
 
@@ -86,7 +86,7 @@ The steps below will take you through cloning your own private fork, then compil
 ### Linux
 
 1. [Set up Git](https://help.github.com/articles/set-up-git/) and [fork our repository](https://help.github.com/articles/fork-a-repo/).
-   If you'd prefer not to use Git, download the source as a zip file with the button on the right.
+   If you'd prefer not to use Git, use the 'Download ZIP' button on the right to get the source as a zip file.
 
 1. Open your source folder and run **Setup.sh** to download binary content for the engine.
 
