Update the README for the master branch.

diff --git a/README.md b/README.md
--- a/README.md
+++ b/README.md
@@ -40,7 +40,7 @@ The steps below will take you through cloning your own private fork, then compil
 
 1. Install **[GitHub for Windows](https://windows.github.com/)** then **fork our repository** and **clone it to your computer**. 
    If you'd rather use Git from the command line, see the [Setting up Git](https://help.github.com/articles/set-up-git/) and [Fork a Repo](https://help.github.com/articles/fork-a-repo/) articles.
-   If you'd prefer not to use Git at all, use the 'Download ZIP' button on the right to get the source as a zip file.
+   If you'd prefer not to use Git at all, you can get the source as a zip file using the 'Download ZIP' button on the right.
 
 1. Make sure you have **Visual Studio 2013** installed.
    Any desktop edition of Visual Studio 2013 will do, including the free version: [Visual Studio 2013 Express for Windows Desktop](http://www.microsoft.com/en-us/download/details.aspx?id=40787).
