Update Git README

diff --git a/README.md b/README.md
--- a/README.md
+++ b/README.md
@@ -38,12 +38,15 @@ The steps below will take you through cloning your own private fork, then compil
 
 ### Windows
 
-1. Install **[GitHub for Windows](https://windows.github.com/)** then **fork our repository** and **clone it to your computer**. 
-   If you'd rather use Git from the command line, see the [Setting up Git](https://help.github.com/articles/set-up-git/) and [Fork a Repo](https://help.github.com/articles/fork-a-repo/) articles.
-   If you'd prefer not to use Git at all, you can get the source as a zip file using the 'Download ZIP' button on the right.
+1. Install **[GitHub for Windows](https://windows.github.com/)** then **[fork and clone our repository](https://guides.github.com/activities/forking/)**. 
+   To use Git from the command line, see the [Setting up Git](https://help.github.com/articles/set-up-git/) and [Fork a Repo](https://help.github.com/articles/fork-a-repo/) articles.
+
+   If you'd prefer not to use Git, you can get the source with the 'Download ZIP' button on the right. **Note:** the built-in Windows unzip utility will mark the contents of zip 
+   files downloaded from the Internet as unsafe to execute, so right-click the zip file and select 'Properties...' then 'Unblock' before you decompress it. 
+   Third-party zip utilities don't normally do this.
 
 1. Make sure you have **Visual Studio 2013** installed.
-   Any desktop edition of Visual Studio 2013 will do, including the free version: [Visual Studio 2013 Express for Windows Desktop](http://www.microsoft.com/en-us/download/details.aspx?id=40787).
+   Any desktop edition of Visual Studio 2013 will do, including the free version; [Visual Studio 2013 Express for Windows Desktop](http://www.microsoft.com/en-us/download/details.aspx?id=40787).
 
 1. Open your source folder in Explorer and run **Setup.bat**. This will download binary content for the engine, as well as installing prerequisites
    and setting up Unreal file associations. On Windows 8, a warning from SmartScreen may appear.  Click "More info", then "Run anyway" to continue.
@@ -63,22 +66,23 @@ The steps below will take you through cloning your own private fork, then compil
 
 ### Mac
    
-1. Install **[GitHub for Mac](https://mac.github.com/)** then **fork our repository** and **clone it to your computer**. 
-   To use Git from the Terminal, see the [Setting up Git](https://help.github.com/articles/set-up-git/) and [Fork a Repo](https://help.github.com/articles/fork-a-repo/) articles. 
-   Or if you'd rather not use Git at all, use the 'Download ZIP' button on the right to get the source as a zip file.
+1. Install **[GitHub for Mac](https://mac.github.com/)** then **[fork and clone our repository](https://guides.github.com/activities/forking/)**. 
+   To use Git from the Terminal, see the [Setting up Git](https://help.github.com/articles/set-up-git/) and [Fork a Repo](https://help.github.com/articles/fork-a-repo/) articles.
+   If you'd rather not use Git at all, use the 'Download ZIP' button on the right to get the source as a zip file.
 
 1. Install the latest version of [Xcode](https://itunes.apple.com/us/app/xcode/id497799835).
 
-1. Open your source folder in Finder and run **Setup.command** to download binary content for the engine.
-
-1. You'll need project files in order to compile.  In the _UnrealEngine_ folder, double-click on **GenerateProjectFiles.command**.  It should take less than a minute to complete.  You can close the Terminal window afterwards.  If you downloaded the source in .zip format, you may see a warning about an unidentified developer.  This is because because the .zip files on GitHub are not digitally signed.  To work around this, right-click on GenerateProjectFiles.command, select Open, then click the Open button if you are sure you want to open it.
+1. Open your source folder in Finder and double-click on **Setup.command** to download binary content for the engine. You can close the Terminal window afterwards.
 
-1. Load the project into Xcode by double-clicking on the **UE4.xcodeproj** file.
+   If you downloaded the source as a .zip file, you may see a warning about it being from an unidentified developer (because .zip files on GitHub aren't digitally signed).
+   To work around it, right-click on GenerateProjectFiles.command, select Open, then click the Open button.
 
-1. It's time to **compile the editor**!  In Xcode, build the **UE4Editor - Mac** target for **My Mac** by selecting it in the target drop down
-   in the top left, then using **Product** -> **Build For** -> **Running**.  It will take between 15 and 40 minutes, depending on your system specs.
+1. In the same folder, double-click **GenerateProjectFiles.command**.  It should take less than a minute to complete.  
 
-1. After compiling finishes, you can **load the editor** from Xcode using the **Product** -> **Run** menu command!
+1. Load the project into Xcode by double-clicking on the **UE4.xcodeproj** file. Select the **UE4Editor - Mac** for **My Mac** target in the title bar,
+   then select the 'Product > Build' menu item. Compiling may take anywhere between 15 and 40 minutes, depending on your system specs.
+   
+1. After compiling finishes, select the 'Product > Run' menu item to load the editor.
 
 
 
