More minor README tweaks.

diff --git a/README.md b/README.md
--- a/README.md
+++ b/README.md
@@ -41,17 +41,20 @@ The steps below will take you through cloning your own private fork, then compil
 1. Install **[GitHub for Windows](https://windows.github.com/)** then **[fork and clone our repository](https://guides.github.com/activities/forking/)**. 
    To use Git from the command line, see the [Setting up Git](https://help.github.com/articles/set-up-git/) and [Fork a Repo](https://help.github.com/articles/fork-a-repo/) articles.
 
-   If you'd prefer not to use Git, you can get the source with the 'Download ZIP' button on the right. **Note:** the built-in Windows unzip utility will mark the contents of zip 
-   files downloaded from the Internet as unsafe to execute, so right-click the zip file and select 'Properties...' then 'Unblock' before you decompress it. 
-   Third-party zip utilities don't normally do this.
+   If you'd prefer not to use Git, you can get the source with the 'Download ZIP' button on the right. 
+   
+   **Note:** The built-in Windows unzip utility will mark the contents of zip files downloaded from the Internet as unsafe to execute, so right-click the zip file and 
+   select 'Properties...' then 'Unblock' before you decompress it. Third-party zip utilities don't normally do this.
 
 1. Make sure you have **Visual Studio 2013** installed.
    Any desktop edition of Visual Studio 2013 will do, including the free version; [Visual Studio 2013 Express for Windows Desktop](http://www.microsoft.com/en-us/download/details.aspx?id=40787).
 
-1. Open your source folder in Explorer and run **Setup.bat**. This will download binary content for the engine, as well as installing prerequisites
-   and setting up Unreal file associations. On Windows 8, a warning from SmartScreen may appear.  Click "More info", then "Run anyway" to continue.
+1. Open your source folder in Explorer and run **Setup.bat**. 
+
+   This will download binary content for the engine, as well as installing prerequisites and setting up Unreal file associations. 
+   On Windows 8, a warning from SmartScreen may appear.  Click "More info", then "Run anyway" to continue.
    
-   A clean download of the engine binaries is currently around 2.5gb, which may take some time to complete on a slow internet connection.
+   A clean download of the engine binaries is currently around 2.5gb, which may take some time to complete.
    Subsequent checkouts only require incremental downloads and will be much quicker.
  
 1. Run **GenerateProjectFiles.bat** to create project files for the engine. It should take less than a minute to complete.  
