More README updates!

diff --git a/README.md b/README.md
--- a/README.md
+++ b/README.md
@@ -12,7 +12,7 @@ We have a heap of documentation available for the engine on the web. If you're l
 * [Unreal Engine API Reference](https://docs.unrealengine.com/latest/INT/API/index.html)
 * [Engine source and GitHub on the Unreal Engine forums](https://forums.unrealengine.com/forumdisplay.php?1-Development-Discussion)
 
-If you need more, just ask. A lot of Epic developers hang out on the [forums](https://forums.unrealengine.com/) or [AnswerHub](https://answers.unrealengine.com/), 
+If you need more, just ask! A lot of Epic developers hang out on the [forums](https://forums.unrealengine.com/) or [AnswerHub](https://answers.unrealengine.com/), 
 and we're proud to be part of a well-meaning, friendly and welcoming community of thousands. 
 
 
@@ -24,12 +24,13 @@ We publish source for the engine in three rolling branches:
 The **[release branch](https://github.com/EpicGames/UnrealEngine/tree/release)** is extensively tested by our QA team and makes a great starting point for learning the engine or
 making your own games. We work hard to make releases stable and reliable, and aim to publish new releases every 1-2 months.
 
-The **[promoted branch](https://github.com/EpicGames/UnrealEngine/tree/promoted)** branch is updated with builds for our artists and designers to use. We try to update it daily, though we often
-catch things that prevent us from doing so, and it's a good balance between getting the latest cool stuff and knowing most things work.
+The **[promoted branch](https://github.com/EpicGames/UnrealEngine/tree/promoted)** is updated with builds for our artists and designers to use. We try to update it daily 
+(though we often catch things that prevent us from doing so) and it's a good balance between getting the latest cool stuff and knowing most things work.
 
-The **[master branch](https://github.com/EpicGames/UnrealEngine/tree/master)** branch tracks [live changes](https://github.com/EpicGames/UnrealEngine/commits/master) by our engine team. This is the cutting edge and may be buggy - it may not even 
-compile. Battle-hardened developers eager to work lock-step with us on the latest and greatest should head here.
+The **[master branch](https://github.com/EpicGames/UnrealEngine/tree/master)** tracks [live changes](https://github.com/EpicGames/UnrealEngine/commits/master) by our engine team. 
+This is the cutting edge and may be buggy - it may not even compile. Battle-hardened developers eager to work lock-step with us on the latest and greatest should head here.
 
+Other short-lived branches may pop-up from time to time as we stabilize new releases or hotfixes.
 
 Getting up and running
 ----------------------
@@ -41,16 +42,13 @@ The steps below will take you through cloning your own private fork, then compil
 1. Install **[GitHub for Windows](https://windows.github.com/)** then **[fork and clone our repository](https://guides.github.com/activities/forking/)**. 
    To use Git from the command line, see the [Setting up Git](https://help.github.com/articles/set-up-git/) and [Fork a Repo](https://help.github.com/articles/fork-a-repo/) articles.
 
-   If you'd prefer not to use Git, you can get the source with the 'Download ZIP' button on the right. 
-   
-   **Note:** The built-in Windows unzip utility will mark the contents of zip files downloaded from the Internet as unsafe to execute, so right-click the zip file and 
-   select 'Properties...' then 'Unblock' before you decompress it. Third-party zip utilities don't normally do this.
+   If you'd prefer not to use Git, you can get the source with the 'Download ZIP' button on the right. The built-in Windows zip utility will mark the contents of zip files 
+   downloaded from the Internet as unsafe to execute, so right-click the zip file and select 'Properties...' and 'Unblock' before decompressing it. Third-party zip utilities don't normally do this.
 
 1. Make sure you have **Visual Studio 2013** installed.
-   Any desktop edition of Visual Studio 2013 will do, including the free version; [Visual Studio 2013 Express for Windows Desktop](http://www.microsoft.com/en-us/download/details.aspx?id=40787).
+   Any desktop edition of Visual Studio 2013 will do, including the free version - [Visual Studio 2013 Express for Windows Desktop](http://www.microsoft.com/en-us/download/details.aspx?id=40787).
 
 1. Open your source folder in Explorer and run **Setup.bat**. 
-
    This will download binary content for the engine, as well as installing prerequisites and setting up Unreal file associations. 
    On Windows 8, a warning from SmartScreen may appear.  Click "More info", then "Run anyway" to continue.
    
@@ -71,7 +69,7 @@ The steps below will take you through cloning your own private fork, then compil
    
 1. Install **[GitHub for Mac](https://mac.github.com/)** then **[fork and clone our repository](https://guides.github.com/activities/forking/)**. 
    To use Git from the Terminal, see the [Setting up Git](https://help.github.com/articles/set-up-git/) and [Fork a Repo](https://help.github.com/articles/fork-a-repo/) articles.
-   If you'd rather not use Git at all, use the 'Download ZIP' button on the right to get the source as a zip file.
+   If you'd rather not use Git, use the 'Download ZIP' button on the right to get the source directly.
 
 1. Install the latest version of [Xcode](https://itunes.apple.com/us/app/xcode/id497799835).
 
