Update link to VS2013 (PR #624 from dthielman).

diff --git a/README.md b/README.md
--- a/README.md
+++ b/README.md
@@ -46,7 +46,7 @@ The steps below will take you through cloning your own private fork, then compil
    downloaded from the Internet as unsafe to execute, so right-click the zip file and select 'Properties...' and 'Unblock' before decompressing it. Third-party zip utilities don't normally do this.
 
 1. Make sure you have **Visual Studio 2013** installed.
-   Any desktop edition of Visual Studio 2013 will do, including the free version - [Visual Studio 2013 Express for Windows Desktop](http://www.microsoft.com/en-us/download/details.aspx?id=40787).
+   Any desktop edition of Visual Studio 2013 will do, including the free version - [Visual Studio Community 2013](http://www.visualstudio.com/products/visual-studio-community-vs).
 
 1. Open your source folder in Explorer and run **Setup.bat**. 
    This will download binary content for the engine, as well as installing prerequisites and setting up Unreal file associations. 
