Update Git README

diff --git a/README.md b/README.md
--- a/README.md
+++ b/README.md
@@ -45,9 +45,9 @@ The steps below will take you through cloning your own private fork, then compil
    If you'd prefer not to use Git, you can get the source with the 'Download ZIP' button on the right. The built-in Windows zip utility will mark the contents of zip files 
    downloaded from the Internet as unsafe to execute, so right-click the zip file and select 'Properties...' and 'Unblock' before decompressing it. Third-party zip utilities don't normally do this.
 
-1. Make sure you have **Visual Studio 2013** installed.
-   Any desktop edition of Visual Studio 2013 will do, including the free version - [Visual Studio Community 2013](http://www.visualstudio.com/products/visual-studio-community-vs).
-
+1. Install **Visual Studio 2013**. Any desktop edition of Visual Studio 2013 can build UE4, including [Visual Studio Community 2013](http://www.visualstudio.com/products/visual-studio-community-vs)
+   which is available for free. Make sure to include the MFC libraries as part of the install (it's enabled by default), which we need for ATL support.
+   
 1. Open your source folder in Explorer and run **Setup.bat**. 
    This will download binary content for the engine, as well as installing prerequisites and setting up Unreal file associations. 
    On Windows 8, a warning from SmartScreen may appear.  Click "More info", then "Run anyway" to continue.
