Fixing .gitignore

diff --git a/.gitignore b/.gitignore
--- a/.gitignore
+++ b/.gitignore
@@ -44,6 +44,7 @@
 !/Engine/Binaries/ThirdParty/Mono/Mac
 /Engine/Binaries/ThirdParty/Mono/Mac/**
 !/Engine/Binaries/ThirdParty/Mono/Mac/bin
+!/Engine/Binaries/ThirdParty/Mono/Mac/bin/*
 
 !/Engine/Binaries/ThirdParty/Mono/Mac/etc
 !/Engine/Binaries/ThirdParty/Mono/Mac/etc/mono
