Snapshot for changelist 2400041.

[CL 2400041 in Main branch]

diff --git a/.gitignore b/.gitignore
--- a/.gitignore
+++ b/.gitignore
@@ -89,6 +89,9 @@
 /Engine/Programs/**/Saved/**
 /Engine/Programs/UnrealBuildTool/**
 
+# Intermediate folders for programs should not be checked in
+/Engine/Programs/**/Intermediate/**
+
 # Ignore various other dependencies that are never checked in
 /Engine/Content/**
 /Engine/Documentation/**
@@ -125,6 +128,8 @@
 /Engine/Source/ThirdParty/**/*.idb
 /Engine/Source/ThirdParty/**/*.bc
 /Engine/Source/ThirdParty/**/*.zip
+/Engine/Source/ThirdParty/**/*.pak
+/Engine/Source/ThirdParty/**/*.framework
 
 # Ignore LinuxNativeDialogs build files (all extensions)
 /Engine/Source/ThirdParty/LinuxNativeDialogs/UELinuxNativeDialogs/build/**
