Tweak README.

diff --git a/README.md b/README.md
--- a/README.md
+++ b/README.md
@@ -77,7 +77,7 @@ The steps below will take you through cloning your own private fork, then compil
 1. Open your source folder in Finder and double-click on **Setup.command** to download binary content for the engine. You can close the Terminal window afterwards.
 
    If you downloaded the source as a .zip file, you may see a warning about it being from an unidentified developer (because .zip files on GitHub aren't digitally signed).
-   To work around it, right-click on GenerateProjectFiles.command, select Open, then click the Open button.
+   To work around it, right-click on Setup.command, select Open, then click the Open button.
 
 1. In the same folder, double-click **GenerateProjectFiles.command**.  It should take less than a minute to complete.  
 
