Snapshot for changelist 2410803.

[CL 2410803 in Main branch]

diff --git a/.gitignore b/.gitignore
--- a/.gitignore
+++ b/.gitignore
@@ -130,6 +130,7 @@
 /Engine/Source/ThirdParty/**/*.zip
 /Engine/Source/ThirdParty/**/*.pak
 /Engine/Source/ThirdParty/**/*.framework
+/Engine/Source/ThirdParty/**/*.tga
 
 # Ignore LinuxNativeDialogs build files (all extensions)
 /Engine/Source/ThirdParty/LinuxNativeDialogs/UELinuxNativeDialogs/build/**
