Remove outdated setup instructions for Linux native compilation.

diff --git a/README.md b/README.md
--- a/README.md
+++ b/README.md
@@ -100,7 +100,7 @@ The steps below will take you through cloning your own private fork, then compil
 
    **Cross-compiling** is handy when you are a Windows (Mac support planned too) developer who wants to package your game for Linux with minimal hassle, and it requires a [cross-compiler toolchain](http://cdn.unrealengine.com/qfe/v4_clang-3.5.0_ld-2.24_glibc-2.12.2.zip) to be installed (see the [Linux cross-compiling page on the wiki](https://wiki.unrealengine.com/Compiling_For_Linux)).
 
-   **Native compilation** is discussed in [a separate README](Engine/Build/BatchFiles/Linux/README.md) and [community wiki page](https://wiki.unrealengine.com/Building_On_Linux). Downloading the dependencies has now been automated, so you will only need to clone the repo and run **GenerateProjectFiles.sh** (provided that you have OAUTH_TOKEN set to your personal access token, see the above README for details).
+   **Native compilation** is discussed in [a separate README](Engine/Build/BatchFiles/Linux/README.md) and [community wiki page](https://wiki.unrealengine.com/Building_On_Linux). 
 
 
 
