Deprecate support for VS2012.

diff --git a/README.md b/README.md
--- a/README.md
+++ b/README.md
@@ -120,8 +120,6 @@ The steps below will take you through cloning your own private fork, then compil
 Additional Notes
 ----------------
 
-Visual Studio 2012 is supported by the Windows toolchain, though Visual Studio 2013 is recommended.
-
 The first time you start the editor from a fresh source build, you may experience long load times. 
 The engine is optimizing content for your platform to the _derived data cache_, and it should only happen once.
 
