Keep re-running GitDependencies if it crashed (UE-7657).

diff --git a/Setup.sh b/Setup.sh
--- a/Setup.sh
+++ b/Setup.sh
@@ -28,11 +28,11 @@ else
 	if [ -d .git/hooks ]; then
 		echo "Registering git hooks... (this will override existing ones!)"
 		echo \#!/bin/sh >.git/hooks/post-checkout
-		echo mono Engine/Binaries/DotNET/GitDependencies.exe >>.git/hooks/post-checkout
+		echo Engine/Binaries/DotNET/GitDependencies.sh >>.git/hooks/post-checkout
 		chmod +x .git/hooks/post-checkout
 
 		echo \#!/bin/sh >.git/hooks/post-merge
-		echo mono Engine/Binaries/DotNET/GitDependencies.exe >>.git/hooks/post-merge
+		echo Engine/Binaries/DotNET/GitDependencies.sh >>.git/hooks/post-merge
 		chmod +x .git/hooks/post-merge
 	fi
 
