Update README with info on user contributions. PR #809

diff --git a/README.md b/README.md
--- a/README.md
+++ b/README.md
@@ -117,6 +117,17 @@ The steps below will take you through cloning your own private fork, then compil
 
 
 
+Contributing
+------------
+
+We welcome any contributions to Unreal Engine development through [pull requests](https://help.github.com/articles/using-pull-requests/) on GitHub. 
+
+Most of our active development is in the **master** branch, so we prefer to take pull requests there (particularly for new features). We try to make sure that all new code adheres to the [Epic coding standards](https://docs.unrealengine.com/latest/INT/Programming/Development/CodingStandard/). 
+
+Contributions are governed by terms in the [Unreal Engine EULA](https://www.unrealengine.com/eula).
+
+
+
 Additional Notes
 ----------------
 
