Snapshot for changelist 2460081.

[CL 2460081 in Main branch]

diff --git a/.gitignore b/.gitignore
--- a/.gitignore
+++ b/.gitignore
@@ -138,6 +138,7 @@
 # Ignore sample and template projects
 /Samples/**
 /Templates/**
+/FeaturePacks/**
 
 # Ignore Unix backup files
 *~
