PR #877: Fix typo (Contributed by [REDACTED_PERSON])

https://github.com/EpicGames/UnrealEngine/pull/877

diff --git a/UE4Games.uprojectdirs b/UE4Games.uprojectdirs
--- a/UE4Games.uprojectdirs
+++ b/UE4Games.uprojectdirs
@@ -1,4 +1,4 @@
-; These folders were be searched 1 level deep in order to find projects
+; These folders will be searched 1 level deep in order to find projects
 ; UnrealBuildTool will store the following information:
 ; - Project name
 ; - Location of project
@@ -6,4 +6,4 @@
 ; - TargetNames contains at the project location
 ;
 ./
-Engine/Source/
\ No newline at end of file
+Engine/Source/
