[INTEGRATE] Change 2474608 by [REDACTED_PERSON][REDACTED_HANDLE] on 2015/03/10 19:57:47

	PR #909: Update .gitignore (Contributed by [REDACTED_PERSON])

	https://github.com/EpicGames/UnrealEngine/pull/909

diff --git a/.gitignore b/.gitignore
--- a/.gitignore
+++ b/.gitignore
@@ -154,6 +154,9 @@
 # Ignore crash reports
 /crashinfo--*
 
+# Ignore package version counter
+/Engine/Build/IOS/UE4Client.PackageVersionCounter
+
 # Ignore linux project files
 *.user
 *.pro
