PR #926: Update size of external binaries (Contributed by [REDACTED_PERSON])

diff --git a/README.md b/README.md
--- a/README.md
+++ b/README.md
@@ -53,7 +53,7 @@ The steps below will take you through cloning your own private fork, then compil
    This will download binary content for the engine, as well as installing prerequisites and setting up Unreal file associations. 
    On Windows 8, a warning from SmartScreen may appear.  Click "More info", then "Run anyway" to continue.
    
-   A clean download of the engine binaries is currently around 2.5gb, which may take some time to complete.
+   A clean download of the engine binaries is currently 3-4gb, which may take some time to complete.
    Subsequent checkouts only require incremental downloads and will be much quicker.
  
 1. Run **GenerateProjectFiles.bat** to create project files for the engine. It should take less than a minute to complete.  
