Add a more useful error message if user tries to run GenerateProjectFiles.sh without running Setup.sh first.

diff --git a/GenerateProjectFiles.sh b/GenerateProjectFiles.sh
--- a/GenerateProjectFiles.sh
+++ b/GenerateProjectFiles.sh
@@ -11,6 +11,13 @@ if [ ! -f Engine/Build/BatchFiles/Mac/GenerateProjectFiles.sh ]; then
   exit 1
 fi 
 
+if [ -f Setup.sh ]; then
+	if [ ! -f .ue4dependencies ]; then
+		echo "Please run Setup to download dependencies before generating project files."
+		exit 1
+	fi
+fi
+
 if [ "$(uname)" = "Darwin" ]; then
 	cd Engine/Build/BatchFiles/Mac
 	sh ./GenerateLLDBInit.sh
