Update legal text in README.

diff --git a/README.md b/README.md
--- a/README.md
+++ b/README.md
@@ -117,14 +117,12 @@ The steps below will take you through cloning your own private fork, then compil
 
 
 
-Contributing
-------------
+Licensing and Contributions
+---------------------------
 
-We welcome any contributions to Unreal Engine development through [pull requests](https://help.github.com/articles/using-pull-requests/) on GitHub. 
+Your access to and use of Unreal Engine on GitHub is governed by the [Unreal Engine End User License Agreement](https://www.unrealengine.com/eula). If you don't agree to those terms, as amended from time to time, you are not permitted to access or use Unreal Engine.
 
-Most of our active development is in the **master** branch, so we prefer to take pull requests there (particularly for new features). We try to make sure that all new code adheres to the [Epic coding standards](https://docs.unrealengine.com/latest/INT/Programming/Development/CodingStandard/). 
-
-Contributions are governed by terms in the [Unreal Engine EULA](https://www.unrealengine.com/eula).
+We welcome any contributions to Unreal Engine development through [pull requests](https://help.github.com/articles/using-pull-requests/) on GitHub. Most of our active development is in the **master** branch, so we prefer to take pull requests there (particularly for new features). We try to make sure that all new code adheres to the [Epic coding standards](https://docs.unrealengine.com/latest/INT/Programming/Development/CodingStandard/).  All contributions are governed by the terms of the EULA.
 
 
 
