Ignore build receipts in Git. They go along with the binaries, so we never want to include them in the repository.

diff --git a/.gitignore b/.gitignore
--- a/.gitignore
+++ b/.gitignore
@@ -118,6 +118,9 @@
 # Ignore any saved local files
 /Engine/Saved/**
 
+# Ignore any build receipts
+/Engine/Build/Receipts/**
+
 # Install stamp linux dependencies
 /Engine/Build/BatchFiles/Linux/dependencies.sha1sums.installed
 
