Mask ReplayServer binaries out of Git.

diff --git a/.gitignore b/.gitignore
--- a/.gitignore
+++ b/.gitignore
@@ -153,6 +153,9 @@
 # Ignore LinuxNativeDialogs build files (all extensions)
 /Engine/Source/ThirdParty/LinuxNativeDialogs/UELinuxNativeDialogs/build/**
 
+# Ignore all ReplaySever JARs
+/Engine/Source/Programs/ReplayServer/ThirdParty/**/*.jar
+
 # Ignore PhysX binaries
 /Engine/Source/ThirdParty/PhysX/APEX-1.3/externals/CUDA_6.5.19/bin/
 /Engine/Source/ThirdParty/PhysX/PhysX-3.3/externals/clang/3.3.2/linux32/bin/
