PR #1268: Added *.dat to ignored third-party binaries (Contributed by [REDACTED_PERSON])

diff --git a/.gitignore b/.gitignore
--- a/.gitignore
+++ b/.gitignore
@@ -146,6 +146,7 @@
 /Engine/Source/ThirdParty/**/*.tga
 /Engine/Source/ThirdParty/**/*.tiff
 /Engine/Source/ThirdParty/**/*.icns
+/Engine/Source/ThirdParty/**/*.dat
 
 # Ignore PLCrashReporter binaries
 /Engine/Source/ThirdParty/PLCrashReporter/**/bin/**
