Update README for the release branch.

diff --git a/README.md b/README.md
--- a/README.md
+++ b/README.md
@@ -22,7 +22,7 @@ Branches
 We publish source for the engine in three rolling branches:
 
 The **[release branch](https://github.com/EpicGames/UnrealEngine/tree/release)** is extensively tested by our QA team and makes a great starting point for learning the engine or
-making your own games. We work hard to make releases stable and reliable, and aim to publish new releases every 1-2 months.
+making your own games. We work hard to make releases stable and reliable, and aim to publish new releases every few months.
 
 The **[promoted branch](https://github.com/EpicGames/UnrealEngine/tree/promoted)** is updated with builds for our artists and designers to use. We try to update it daily 
 (though we often catch things that prevent us from doing so) and it's a good balance between getting the latest cool stuff and knowing most things work.
