PR #1390: Setup.bat improvement (Contributed by [REDACTED_PERSON])

diff --git a/Setup.bat b/Setup.bat
--- a/Setup.bat
+++ b/Setup.bat
@@ -10,9 +10,9 @@ rem Setup the git hooks...
 if not exist .git\hooks goto no_git_hooks_directory
 echo Registering git hooks...
 echo #!/bin/sh >.git\hooks\post-checkout
-echo Engine/Binaries/DotNET/GitDependencies.exe >>.git\hooks\post-checkout
+echo Engine/Binaries/DotNET/GitDependencies.exe %* >>.git\hooks\post-checkout
 echo #!/bin/sh >.git\hooks\post-merge
-echo Engine/Binaries/DotNET/GitDependencies.exe >>.git\hooks\post-merge
+echo Engine/Binaries/DotNET/GitDependencies.exe %* >>.git\hooks\post-merge
 :no_git_hooks_directory
 
 rem Install prerequisites...
