Switch Git repository to whitelist repository files rather than blacklisting them. Should prevent unwanted junk from newly added third party libraries creeping into the repository.

diff --git a/.gitignore b/.gitignore
--- a/.gitignore
+++ b/.gitignore
@@ -13,174 +13,133 @@
 #  Use ! to negate a previous pattern.  But it doesn't work if the parent sub-folder was masked out already.
 # 
 
+# Ignore all files by default, but scan all directories
+*
+!*/
+
+# C/C++ source files
+!*.c
+!*.cc
+!*.cpp
+!*.cpp.template
+!*.h
+!*.h.template
+!*.hpp
+!*.inl
+!*.inc
+!*.m
+!*.mm
+!*.rc
+!*.def
+!*.exp
+!*.manifest
+
+# Java source files
+!*.java
+!*.java.template
+
+# C# source files
+!*.cs
+!*.cs.template
+!*.aspx
+!*.resx
+
+# Shader formats
+!*.usf
+!*.hlsl
+!*.glsl
+
+# Text files
+!*.txt
+!*.md
+
+# Script files
+!*.bat
+!*.sh
+!*.pl
+!*.py
+!*.js
+!*.command
+
+# Other configuration and markup files
+!*.ini
+!*.json
+!*.tps
+!*.xml
+!*.xaml
+!*.uproject
+!*.uplugin
+!*.html
+!*.html.template
+!*.css
+!*.udn
+!*.config
+!*.version
+!.git*
+
+# Projects and makefiles
+!*.cmake
+!*.mk
+!*.dsp
+!*.dsw
+!*.csproj
+!*.vcproj
+!*.vcxproj
+!*.vcxproj.filters
+!*.sln
+!*.xcodeproj
+!*.xcconfig
+!*.vsprops
+!*.snippet
+!Makefile
+!Makefile.*
+!Settings.settings
+
+# Specific names
+!README
+!AUTHORS
+!LICENSE
+!FAQ
+!VERSION
+!ChangeLog
 
-# Ignore project files in the root
+# Ignore Unix backup files
+*~
+
+# Exceptions
 /*.sln
-*.suo
-*.opensdf
-*.sdf
 /*.xcodeproj
 /Makefile
 /CMakeLists.txt
-/.ue4dependencies
-
-# Derived data cache is never checked in
-/Engine/DerivedDataCache/**
-
-# Engine intermediates
-/Engine/Intermediate/**
-/Engine/Plugins/**/Intermediate/**
-
-# Ignore binary files...
-/Engine/Binaries/**
-/Engine/Plugins/**/Binaries/**
-
-# Add exceptions for Mac Mono binaries needed to run GitDependencies.exe. This is a mess, but we need to un-ignore every parent directory before anything in it,
-# and Mono has a pretty deep hierarchy.
-!/Engine/Binaries/ThirdParty
-/Engine/Binaries/ThirdParty/**
-!/Engine/Binaries/ThirdParty/Mono
-/Engine/Binaries/ThirdParty/Mono/**
-!/Engine/Binaries/ThirdParty/Mono/Mac
-/Engine/Binaries/ThirdParty/Mono/Mac/**
-!/Engine/Binaries/ThirdParty/Mono/Mac/bin
-!/Engine/Binaries/ThirdParty/Mono/Mac/bin/*
+*.target.xml
 
-!/Engine/Binaries/ThirdParty/Mono/Mac/etc
-!/Engine/Binaries/ThirdParty/Mono/Mac/etc/mono
-/Engine/Binaries/ThirdParty/Mono/Mac/etc/mono/**
+# Add exceptions for Mac Mono binaries needed to run GitDependencies.exe.
+!/Engine/Binaries/ThirdParty/Mono/Mac/bin/*
 !/Engine/Binaries/ThirdParty/Mono/Mac/etc/mono/config
-!/Engine/Binaries/ThirdParty/Mono/Mac/etc/mono/4.5
-/Engine/Binaries/ThirdParty/Mono/Mac/etc/mono/4.5/**
 !/Engine/Binaries/ThirdParty/Mono/Mac/etc/mono/4.5/machine.config
-
-!/Engine/Binaries/ThirdParty/Mono/Mac/lib/
-/Engine/Binaries/ThirdParty/Mono/Mac/lib/**
 !/Engine/Binaries/ThirdParty/Mono/Mac/lib/libMonoPosixHelper.dylib
-!/Engine/Binaries/ThirdParty/Mono/Mac/lib/mono
-/Engine/Binaries/ThirdParty/Mono/Mac/lib/mono/**
-!/Engine/Binaries/ThirdParty/Mono/Mac/lib/mono/4.5
-/Engine/Binaries/ThirdParty/Mono/Mac/lib/mono/4.5/**
 !/Engine/Binaries/ThirdParty/Mono/Mac/lib/mono/4.5/mscorlib.dll
-
-!/Engine/Binaries/ThirdParty/Mono/Mac/lib/mono/gac
-/Engine/Binaries/ThirdParty/Mono/Mac/lib/mono/gac/**
-!/Engine/Binaries/ThirdParty/Mono/Mac/lib/mono/gac/Mono.Posix
-!/Engine/Binaries/ThirdParty/Mono/Mac/lib/mono/gac/Mono.Security
-!/Engine/Binaries/ThirdParty/Mono/Mac/lib/mono/gac/System
-!/Engine/Binaries/ThirdParty/Mono/Mac/lib/mono/gac/System.Configuration
-!/Engine/Binaries/ThirdParty/Mono/Mac/lib/mono/gac/System.Core
-!/Engine/Binaries/ThirdParty/Mono/Mac/lib/mono/gac/System.Xml
-/Engine/Binaries/ThirdParty/Mono/Mac/lib/mono/gac/**/*.mdb
+!/Engine/Binaries/ThirdParty/Mono/Mac/lib/mono/gac/Mono.Posix/4.0.0.0__0738eb9f132ed756/Mono.Posix.dll
+!/Engine/Binaries/ThirdParty/Mono/Mac/lib/mono/gac/Mono.Security/4.0.0.0__0738eb9f132ed756/Mono.Security.dll
+!/Engine/Binaries/ThirdParty/Mono/Mac/lib/mono/gac/System/4.0.0.0__b77a5c561934e089/System.dll
+!/Engine/Binaries/ThirdParty/Mono/Mac/lib/mono/gac/System.Configuration/4.0.0.0__b03f5f7f11d50a3a/System.Configuration.dll
+!/Engine/Binaries/ThirdParty/Mono/Mac/lib/mono/gac/System.Core/4.0.0.0__b77a5c561934e089/System.Core.dll
+!/Engine/Binaries/ThirdParty/Mono/Mac/lib/mono/gac/System.Xml/4.0.0.0__b77a5c561934e089/System.Xml.dll
 
 # ...and GitDependencies.exe
-!/Engine/Binaries/DotNET
-/Engine/Binaries/DotNET/**
 !/Engine/Binaries/DotNET/GitDependencies.exe
 !/Engine/Binaries/DotNET/GitDependencies.exe.config
 
-# This is autogenerated symlink
-/Engine/Binaries/ThirdParty/Mono/Mac/bin/mono
+# Ignore content folders
+Content/
 
-# Intermediate folders created for various C# programs
-/Engine/Source/Programs/**/obj/**
-
-# Saved folders for programs should not be checked in
-/Engine/Programs/**/Saved/**
-/Engine/Programs/UnrealBuildTool/**
+# Ignore DDC
+DerivedDataCache/
 
-# Intermediate folders for programs should not be checked in
-/Engine/Programs/**/Intermediate/**
-
-# Ignore various other dependencies that are never checked in
-/Engine/Content/**
-/Engine/Documentation/**/*.psd
-/Engine/Documentation/**/*.png
-/Engine/Documentation/**/*.cur
-/Engine/Documentation/**/*.gif
-/Engine/Documentation/**/*.jpg
-/Engine/Documentation/CHM/**
-
-# Ignore all binary types in the extras folder
-/Engine/Extras/*/**
-!/Engine/Extras/Maya_AnimationRiggingTools/**
-/Engine/Extras/Maya_AnimationRiggingTools/**/*.pdf
-/Engine/Extras/Maya_AnimationRiggingTools/**/*.pyd
-/Engine/Extras/Maya_AnimationRiggingTools/MayaTools/General/ART/**
-/Engine/Extras/Maya_AnimationRiggingTools/MayaTools/General/Icons/**
-/Engine/Extras/Maya_AnimationRiggingTools/ArtToolsOSX/MayaTools/General/ART/
-/Engine/Extras/Maya_AnimationRiggingTools/ArtToolsOSX/MayaTools/General/Icons/
-/Engine/Extras/Maya_AnimationRiggingTools/ArtToolsOSX/MayaTools/General/Scripts/projectSettings.txt
-/Engine/Extras/Maya_AnimationRiggingTools/ArtToolsOSX/MayaTools/General/Scripts/settings.txt
-/Engine/Extras/Maya_AnimationRiggingTools/ArtToolsOSX/README.md
-!/Engine/Extras/VisualStudioSnippets/**
+# Ignore intermediate folders
+Intermediate/
+obj/
 
 # Ignore any saved local files
-/Engine/Saved/**
-
-# Ignore any build receipts
-/Engine/Build/Receipts/**
-
-# Install stamp linux dependencies
-/Engine/Build/BatchFiles/Linux/dependencies.sha1sums.installed
-
-# log files for linux scripts (e.g. BuildThirdParty.sh)
-/Engine/Build/BatchFiles/Linux/*.log
-
-# Ignore third party binaries
-/Engine/Source/ThirdParty/**/*.exe
-/Engine/Source/ThirdParty/**/*.pdb
-/Engine/Source/ThirdParty/**/*.lib
-/Engine/Source/ThirdParty/**/*.a
-/Engine/Source/ThirdParty/**/*.obj
-/Engine/Source/ThirdParty/**/*.dll
-/Engine/Source/ThirdParty/**/*.dylib
-/Engine/Source/ThirdParty/**/*.bmp
-/Engine/Source/ThirdParty/**/*.so
-/Engine/Source/ThirdParty/**/*.bsc
-/Engine/Source/ThirdParty/**/*.idb
-/Engine/Source/ThirdParty/**/*.bc
-/Engine/Source/ThirdParty/**/*.zip
-/Engine/Source/ThirdParty/**/*.pak
-/Engine/Source/ThirdParty/**/*.framework
-/Engine/Source/ThirdParty/**/*.tga
-/Engine/Source/ThirdParty/**/*.tiff
-/Engine/Source/ThirdParty/**/*.icns
-/Engine/Source/ThirdParty/**/*.dat
-
-# Ignore PLCrashReporter binaries
-/Engine/Source/ThirdParty/PLCrashReporter/**/bin/**
-
-# Ignore LinuxNativeDialogs build files (all extensions)
-/Engine/Source/ThirdParty/LinuxNativeDialogs/UELinuxNativeDialogs/build/**
-
-# Ignore all ReplaySever JARs
-/Engine/Source/Programs/ReplayServer/ThirdParty/**/*.jar
-
-# Ignore PhysX binaries
-/Engine/Source/ThirdParty/PhysX/APEX-1.3/externals/CUDA_6.5.19/bin/
-/Engine/Source/ThirdParty/PhysX/PhysX-3.3/externals/clang/3.3.2/linux32/bin/
-/Engine/Source/ThirdParty/PhysX/PhysX-3.3/externals/clang/3.3.2/osx/bin/
-
-# Ignore sample and template projects
-/Samples/**
-/Templates/**
-/FeaturePacks/**
-
-# Ignore Unix backup files
-*~
-
-# Ignore Mac desktop services store files 
-.DS_Store
-
-# Ignore crash reports
-/crashinfo--*
-
-# Ignore package version counter
-/Engine/Build/IOS/UE4Client.PackageVersionCounter
-
-# Ignore linux project files
-*.user
-*.pro
-*.pri
-*.kdev4
+Saved/
+/Engine/Programs/UnrealBuildTool/**
