Ignore executable config files.

diff --git a/.gitignore b/.gitignore
--- a/.gitignore
+++ b/.gitignore
@@ -112,6 +112,7 @@
 /Makefile
 /CMakeLists.txt
 *.target.xml
+*.exe.config
 
 # Add exceptions for Mac Mono binaries needed to run GitDependencies.exe.
 !/Engine/Binaries/ThirdParty/Mono/Mac/bin/*
