Add mono symlink to gitignore

diff --git a/.gitignore b/.gitignore
--- a/.gitignore
+++ b/.gitignore
@@ -127,6 +127,9 @@
 !/Engine/Binaries/ThirdParty/Mono/Mac/lib/mono/gac/System.Core/4.0.0.0__b77a5c561934e089/System.Core.dll
 !/Engine/Binaries/ThirdParty/Mono/Mac/lib/mono/gac/System.Xml/4.0.0.0__b77a5c561934e089/System.Xml.dll
 
+# Ignore the automatically created mono symlink
+/Engine/Binaries/ThirdParty/Mono/Mac/bin/mono
+
 # ...and GitDependencies.exe
 !/Engine/Binaries/DotNET/GitDependencies.exe
 !/Engine/Binaries/DotNET/GitDependencies.exe.config
