Add explicit .DS_Store to gitignore list; GitHub client on Mac still seems to find them.

diff --git a/.gitignore b/.gitignore
--- a/.gitignore
+++ b/.gitignore
@@ -147,3 +147,6 @@ obj/
 # Ignore any saved local files
 Saved/
 /Engine/Programs/UnrealBuildTool/**
+
+# Explcitly ignore Mac DS_Store files, regardless of where they are
+.DS_Store
