Only ignore /Engine/DerivedDataCache in the Git repository.

diff --git a/.gitignore b/.gitignore
--- a/.gitignore
+++ b/.gitignore
@@ -138,7 +138,7 @@
 Content/
 
 # Ignore DDC
-DerivedDataCache/
+/Engine/DerivedDataCache/**
 
 # Ignore intermediate folders
 Intermediate/
