Fix errorlevel not getting passed as the exit code from CMD.EXE when running GenerateProjectFiles.bat.

diff --git a/GenerateProjectFiles.bat b/GenerateProjectFiles.bat
--- a/GenerateProjectFiles.bat
+++ b/GenerateProjectFiles.bat
@@ -4,11 +4,9 @@ cmd.exe /v /c "%~dp0Engine\Build\BatchFiles\CopyVisualizers.bat"
 
 if not exist "%~dp0Engine\Build\BatchFiles\GenerateProjectFiles.bat" goto Error_BatchFileInWrongLocation
 call "%~dp0Engine\Build\BatchFiles\GenerateProjectFiles.bat" %*
-goto Exit
+exit /B %ERRORLEVEL%
 
 :Error_BatchFileInWrongLocation
 echo GenerateProjectFiles ERROR: The batch file does not appear to be located in the root UE4 directory.  This script must be run from within that directory.
 pause
-goto Exit
-
-:Exit
+exit /B 1
