Don't append the --prompt argument to GitDependencies.exe if it's already specified (or if --force is already specified).

diff --git a/Setup.bat b/Setup.bat
--- a/Setup.bat
+++ b/Setup.bat
@@ -2,8 +2,15 @@
 setlocal
 pushd %~dp0
 
+rem Figure out if we should append the -prompt argument
+set PROMPT_ARGUMENT=
+for %%P in (%*) do if /I "%%P" == "--prompt" goto no_prompt_argument
+for %%P in (%*) do if /I "%%P" == "--force" goto no_prompt_argument
+set PROMPT_ARGUMENT=--prompt
+:no_prompt_argument
+
 rem Sync the dependencies...
-.\Engine\Binaries\DotNET\GitDependencies.exe --prompt %*
+.\Engine\Binaries\DotNET\GitDependencies.exe %PROMPT_ARGUMENT% %*
 if ERRORLEVEL 1 goto error
 
 rem Setup the git hooks...
