Add C# documentation files to .gitignore.

diff --git a/.gitignore b/.gitignore
--- a/.gitignore
+++ b/.gitignore
@@ -136,6 +136,10 @@
 !/Engine/Binaries/DotNET/GitDependencies.exe
 !/Engine/Binaries/DotNET/GitDependencies.exe.config
 
+# Ignore documentation generated for C# tools
+Engine/Binaries/DotNET/UnrealBuildTool.xml
+Engine/Binaries/DotNET/AutomationScripts/BuildGraph.Automation.xml
+
 # Ignore content folders
 Content/
 
