Branch snapshot for CL 3405357

[CL 3405357 in Main branch]

diff --git a/LICENSE.md b/LICENSE.md
new file mode 100644
--- /dev/null
+++ b/LICENSE.md
@@ -0,0 +1 @@
+Use of the Unreal Engine is governed by the terms of the Unreal® Engine End User License Agreement, which can be found at [https://www.unrealengine.com/eula](https://www.unrealengine.com/eula).
