Prevent .ush files being excluded from the Git repo.

diff --git a/.gitignore b/.gitignore
--- a/.gitignore
+++ b/.gitignore
@@ -47,6 +47,7 @@
 
 # Shader formats
 !*.usf
+!*.ush
 !*.hlsl
 !*.glsl
 
