PR #3965: Ignore version files in Binaries directory created by UBT (Contributed by [REDACTED_PERSON])

diff --git a/.gitignore b/.gitignore
--- a/.gitignore
+++ b/.gitignore
@@ -147,6 +147,9 @@
 Engine/Binaries/DotNET/UnrealBuildTool.xml
 Engine/Binaries/DotNET/AutomationScripts/BuildGraph.Automation.xml
 
+# Ignore version files in the Engine/Binaries directory created by UBT
+/Engine/Binaries/**/*.version
+
 # Ignore content folders
 Content/
 
