Branch snapshot for CL 3668440

[CL 3668440 in Staging-4.18 branch]

diff --git a/.editorconfig b/.editorconfig
new file mode 100644
--- /dev/null
+++ b/.editorconfig
@@ -0,0 +1,21 @@
+# EditorConfig is awesome:http://EditorConfig.org
+
+# top-most EditorConfig file
+root = true
+
+# Baseline
+[*]
+charset = utf-8
+indent_style = tab
+trim_trailing_whitespace = true
+max_line_length = 120
+
+# MSBuild
+[*.{csproj,proj,projitems,shproj,fsproj,target,props}]
+indent_style = space
+indent_size = 2
+
+# XML config files
+[*.{config,nuspec,resx}]
+indent_style = space
+indent_size = 2
