PR #4223: .gitignore some more local files (Contributed by [REDACTED_PERSON])

diff --git a/.gitignore b/.gitignore
--- a/.gitignore
+++ b/.gitignore
@@ -150,6 +150,9 @@ Engine/Binaries/DotNET/AutomationScripts/BuildGraph.Automation.xml
 # Ignore version files in the Engine/Binaries directory created by UBT
 /Engine/Binaries/**/*.version
 
+# Ignore exp files in the the Engine/Binaries directory as they aren't C/C++ source files
+/Engine/Binaries/**/*.exp
+
 # Ignore content folders
 Content/
 
@@ -162,9 +165,14 @@ obj/
 
 # Ignore any saved local files
 Saved/
+/Engine/Programs/MemoryProfiler2/**
 /Engine/Programs/UnrealBuildTool/**
 /LocalBuilds/
 
+# Ignore Swarm local save files
+Engine/Binaries/DotNET/SwarmAgent.DeveloperOptions.xml
+Engine/Binaries/DotNET/SwarmAgent.Options.xml
+
 # Explcitly ignore Mac DS_Store files, regardless of where they are
 .DS_Store
 
