PR #4324: Added .vscode in .gitignore (Contributed by [REDACTED_PERSON])

diff --git a/.gitignore b/.gitignore
--- a/.gitignore
+++ b/.gitignore
@@ -178,3 +178,4 @@ Engine/Binaries/DotNET/SwarmAgent.Options.xml
 
 # Ignore VS / VS Code settings directory
 .vs
+.vscode
