PR #4410: Don't gitignore .editorconfig files (Contributed by [REDACTED_PERSON])

diff --git a/.gitignore b/.gitignore
--- a/.gitignore
+++ b/.gitignore
@@ -81,6 +81,7 @@
 !*.config
 !*.version
 !.git*
+!.editorconfig
 
 # Projects and makefiles
 !*.cmake
