PR #4691: Update gitignore (Contributed by [REDACTED_PERSON])

diff --git a/.gitignore b/.gitignore
--- a/.gitignore
+++ b/.gitignore
@@ -118,6 +118,7 @@
 /Makefile
 /CMakeLists.txt
 *.target.xml
+*.dll.config
 *.exe.config
 *.exe.manifest
 /.kdev4/
