Add BuildGraph.Automation.xml to gitignore (#4691). Seems to be generated next to project file as well as in output directory.

diff --git a/.gitignore b/.gitignore
--- a/.gitignore
+++ b/.gitignore
@@ -148,6 +148,7 @@
 # Ignore documentation generated for C# tools
 Engine/Binaries/DotNET/UnrealBuildTool.xml
 Engine/Binaries/DotNET/AutomationScripts/BuildGraph.Automation.xml
+Engine/Source/Programs/AutomationTool/BuildGraph/BuildGraph.Automation.xml
 
 # Ignore version files in the Engine/Binaries directory created by UBT
 /Engine/Binaries/**/*.version
