Git: Update README with info about development streams.

diff --git a/README.md b/README.md
--- a/README.md
+++ b/README.md
@@ -19,16 +19,16 @@ and we're proud to be part of a well-meaning, friendly and welcoming community o
 Branches
 --------
 
-We publish source for the engine in three rolling branches:
+We publish source for the engine in several branches:
 
 The **[release branch](https://github.com/EpicGames/UnrealEngine/tree/release)** is extensively tested by our QA team and makes a great starting point for learning the engine or
 making your own games. We work hard to make releases stable and reliable, and aim to publish new releases every few months.
 
-The **[promoted branch](https://github.com/EpicGames/UnrealEngine/tree/promoted)** is updated with builds for our artists and designers to use. We try to update it daily 
-(though we often catch things that prevent us from doing so) and it's a good balance between getting the latest cool stuff and knowing most things work.
+The **[promoted branch](https://github.com/EpicGames/UnrealEngine/tree/promoted)** is updated with builds for our artists and designers to use. We try to update with merges from the master branch daily (though we often catch things that prevent us from doing so) and it's a good balance between getting the latest cool stuff and knowing most things work.
 
-The **[master branch](https://github.com/EpicGames/UnrealEngine/tree/master)** tracks [live changes](https://github.com/EpicGames/UnrealEngine/commits/master) by our engine team. 
-This is the cutting edge and may be buggy - it may not even compile. Battle-hardened developers eager to work lock-step with us on the latest and greatest should head here.
+The **[master branch](https://github.com/EpicGames/UnrealEngine/tree/master)** is the hub of changes from all our specialized engine development teams. Our internal game teams typically take engine snapshots from here, but it isn't subject to as much testing as release branches.
+
+Individual teams have their own **development branches** for day to day work ([dev-core](https://github.com/EpicGames/UnrealEngine/tree/dev-core), [dev-mobile](https://github.com/EpicGames/UnrealEngine/tree/dev-mobile) and [dev-sequencer](https://github.com/EpicGames/UnrealEngine/tree/dev-sequencer), for example). These branches reflect the cutting edge of the engine and may be buggy - they may not even compile. Battle-hardened developers eager to test new features or work lock-step with us should head to one of these. We aim to merge development branches to master every 3-4 weeks.
 
 Other short-lived branches may pop-up from time to time as we stabilize new releases or hotfixes.
 
