Add "Nuget Package Manager" to the list of required components for compiling the engine.

diff --git a/README.md b/README.md
--- a/README.md
+++ b/README.md
@@ -48,7 +48,7 @@ The steps below will take you through cloning your own private fork, then compil
 
 1. Install **Visual Studio 2017**. 
    All desktop editions of Visual Studio 2017 can build UE4, including [Visual Studio Community 2017](http://www.visualstudio.com/products/visual-studio-community-vs), which is free for small teams and individual developers.
-   To install the correct components for UE4 development, check the "Game Development with C++" workload, and the "Unreal Engine Installer" optional component.
+   To install the correct components for UE4 development, check the "Game Development with C++" workload, and the "Unreal Engine Installer" and "Nuget Package Manager" optional components.
   
 1. Open your source folder in Explorer and run **Setup.bat**. 
    This will download binary content for the engine, as well as installing prerequisites and setting up Unreal file associations. 
