Update .gitignore for Linux Mono changes

diff --git a/.gitignore b/.gitignore
--- a/.gitignore
+++ b/.gitignore
@@ -142,6 +142,26 @@
 # Ignore the automatically created mono symlink
 /Engine/Binaries/ThirdParty/Mono/Mac/bin/mono
 
++# Add exceptions for Linux Mono binaries needed to run GitDependencies.exe.
++!/Engine/Binaries/ThirdParty/Mono/Linux/bin/*
++!/Engine/Binaries/ThirdParty/Mono/Linux/etc/mono/config
++!/Engine/Binaries/ThirdParty/Mono/Linux/etc/mono/4.5/machine.config
++!/Engine/Binaries/ThirdParty/Mono/Linux/x86_64-unknown-linux-gnu/bin/mono-boehm
++!/Engine/Binaries/ThirdParty/Mono/Linux/x86_64-unknown-linux-gnu/bin/mono-sgen
++!/Engine/Binaries/ThirdParty/Mono/Linux/x86_64-unknown-linux-gnu/lib/libMonoPosixHelper.so
++!/Engine/Binaries/ThirdParty/Mono/Linux/lib/mono/4.5/mscorlib.dll
++!/Engine/Binaries/ThirdParty/Mono/Linux/lib/mono/gac/Mono.Posix/4.0.0.0__0738eb9f132ed756/Mono.Posix.dll
++!/Engine/Binaries/ThirdParty/Mono/Linux/lib/mono/gac/Mono.Security/4.0.0.0__0738eb9f132ed756/Mono.Security.dll
++!/Engine/Binaries/ThirdParty/Mono/Linux/lib/mono/gac/System/4.0.0.0__b77a5c561934e089/System.dll
++!/Engine/Binaries/ThirdParty/Mono/Linux/lib/mono/gac/System.Configuration/4.0.0.0__b03f5f7f11d50a3a/System.Configuration.dll
++!/Engine/Binaries/ThirdParty/Mono/Linux/lib/mono/gac/System.Core/4.0.0.0__b77a5c561934e089/System.Core.dll
++!/Engine/Binaries/ThirdParty/Mono/Linux/lib/mono/gac/System.Xml/4.0.0.0__b77a5c561934e089/System.Xml.dll
++
++# Ignore the automatically created mono symlink
++/Engine/Binaries/ThirdParty/Mono/Linux/bin/mono
++/Engine/Binaries/ThirdParty/Mono/Linux/bin/mono-boehm
++/Engine/Binaries/ThirdParty/Mono/Linux/bin/mono-sgen
+
 # ...and GitDependencies.exe
 !/Engine/Binaries/DotNET/GitDependencies.exe
 !/Engine/Binaries/DotNET/GitDependencies.exe.config
