Continue after starting UnrealVersionSelector to avoid blocking a chain command

diff --git a/Setup.sh b/Setup.sh
--- a/Setup.sh
+++ b/Setup.sh
@@ -42,7 +42,7 @@ else
 	echo Register the engine installation...
 	if [ -f Engine/Binaries/Linux/UnrealVersionSelector-Linux-Shipping ]; then
 		pushd Engine/Binaries/Linux > /dev/null
-		./UnrealVersionSelector-Linux-Shipping -register
+		./UnrealVersionSelector-Linux-Shipping -register > /dev/null &
 		popd > /dev/null
 	fi
 
