PR #5330: Capitalise 's' in PlayStation (Contributed by [REDACTED_PERSON])

diff --git a/README.md b/README.md
--- a/README.md
+++ b/README.md
@@ -3,7 +3,7 @@ Unreal Engine
 
 Welcome to the Unreal Engine source code! 
 
-From this repository you can build the Unreal Editor for Windows and Mac, compile Unreal Engine games for Android, iOS, Playstation 4, Xbox One, HTML5 and Linux,
+From this repository you can build the Unreal Editor for Windows and Mac, compile Unreal Engine games for Android, iOS, PlayStation 4, Xbox One, HTML5 and Linux,
 and build tools like Unreal Lightmass and Unreal Frontend. Modify them in any way you can imagine, and share your changes with others! 
 
 We have a heap of documentation available for the engine on the web. If you're looking for the answer to something, you may want to start here: 
@@ -114,7 +114,7 @@ The steps below will take you through cloning your own private fork, then compil
 
 **HTML5** support will be downloaded by the setup script if you have Emscripten installed. Please see the [HTML5 getting started guide](https://docs.unrealengine.com/latest/INT/Platforms/HTML5/GettingStarted/index.html).
 
-**Playstation 4** or **XboxOne** development require additional files that can only be provided after your registered developer status is confirmed by Sony or Microsoft. See [the announcement blog post](https://www.unrealengine.com/blog/playstation-4-and-xbox-one-now-supported) for more information.
+**PlayStation 4** or **XboxOne** development require additional files that can only be provided after your registered developer status is confirmed by Sony or Microsoft. See [the announcement blog post](https://www.unrealengine.com/blog/playstation-4-and-xbox-one-now-supported) for more information.
 
 
 Licensing and Contributions
