PR #5240: Update README.md (Contributed by [REDACTED_PERSON])

diff --git a/README.md b/README.md
--- a/README.md
+++ b/README.md
@@ -3,7 +3,7 @@ Unreal Engine
 
 Welcome to the Unreal Engine source code! 
 
-From this repository you can build the Unreal Editor for Windows and Mac, compile Unreal Engine games for Android, iOS, PlayStation 4, Xbox One, HTML5 and Linux,
+From this repository you can build the Unreal Editor for Windows, Mac and Linux, compile Unreal Engine games for Android, iOS, PlayStation 4, Xbox One and HTML5,
 and build tools like Unreal Lightmass and Unreal Frontend. Modify them in any way you can imagine, and share your changes with others! 
 
 We have a heap of documentation available for the engine on the web. If you're looking for the answer to something, you may want to start here: 
