PR #5529: secure editorconfig.org link (Contributed by [REDACTED_PERSON])

#rnx

diff --git a/.editorconfig b/.editorconfig
--- a/.editorconfig
+++ b/.editorconfig
@@ -1,4 +1,4 @@
-# EditorConfig is awesome: http://EditorConfig.org
+# EditorConfig is awesome: https://EditorConfig.org
 
 # top-most EditorConfig file
 root = true
