Gitignore the clangs toolchain (Thanks PhilBax! PR #5794)

diff --git a/.gitignore b/.gitignore
--- a/.gitignore
+++ b/.gitignore
@@ -209,3 +209,6 @@ Engine/Binaries/DotNET/SwarmAgent.Options.xml
 
 # Annoying windows files
 [Dd]esktop.ini
+
+# Ignore the headers and cpp files in the clang toolchain on Linux
+Engine/Extras/ThirdPartyNotUE/SDKs/HostLinux/**
