FnMerge - Revert single .gitignore after request from Ben.Marsh

diff --git a/.gitignore b/.gitignore
--- a/.gitignore
+++ b/.gitignore
@@ -48,10 +48,6 @@
 !*.aspx
 !*.resx
 
-# ISPC source files
-!*.ispc
-!*.isph
-
 # Shader formats
 !*.usf
 !*.ush
