PR #5343: update gitignore (Contributed by [REDACTED_PERSON])

diff --git a/.gitignore b/.gitignore
--- a/.gitignore
+++ b/.gitignore
@@ -206,3 +206,6 @@ Engine/Binaries/DotNET/SwarmAgent.Options.xml
 
 # Ignore CLion directory
 .idea/
+
+# Annoying windows files
+[Dd]esktop.ini
