Add .natvis to .editorconfig

diff --git a/.editorconfig b/.editorconfig
--- a/.editorconfig
+++ b/.editorconfig
@@ -18,11 +18,11 @@ indent_style = space
 indent_size = 2
 
 # XML config files
-[*.{config,nuspec,resx}]
+[*.{config,nuspec,resx,natvis}]
 indent_style = space
 indent_size = 2
 
 # Python files
 [*.py]
 indent_style = space
-indent_size = 4
\ No newline at end of file
+indent_size = 4
