Install PS4 visualizers from GenerateProjectFiles.bat.  Better error handling.

diff --git a/GenerateProjectFiles.bat b/GenerateProjectFiles.bat
--- a/GenerateProjectFiles.bat
+++ b/GenerateProjectFiles.bat
@@ -1,5 +1,8 @@
 @echo off
 
+REM Install PS4 visualizer if the SDK is present
+cmd.exe /v /c "%~dp0Engine\Extras\VisualStudioDebugging\PS4\InstallPS4Visualizer.bat /quiet"
+
 if not exist "%~dp0Engine\Build\BatchFiles\GenerateProjectFiles.bat" goto Error_BatchFileInWrongLocation
 call "%~dp0Engine\Build\BatchFiles\GenerateProjectFiles.bat" %*
 exit /B %ERRORLEVEL%
