Test for PS4 visualizer installation file before attempting to call it.

diff --git a/GenerateProjectFiles.bat b/GenerateProjectFiles.bat
--- a/GenerateProjectFiles.bat
+++ b/GenerateProjectFiles.bat
@@ -1,7 +1,9 @@
 @echo off
 
-REM Install PS4 visualizer if the SDK is present
-cmd.exe /v /c "%~dp0Engine\Extras\VisualStudioDebugging\PS4\InstallPS4Visualizer.bat /quiet"
+REM Install PS4 visualizer if the SDK and installation file are present
+if exist "%~dp0Engine\Extras\VisualStudioDebugging\PS4\InstallPS4Visualizer.bat" (
+  call "%~dp0Engine\Extras\VisualStudioDebugging\PS4\InstallPS4Visualizer.bat"
+)
 
 if not exist "%~dp0Engine\Build\BatchFiles\GenerateProjectFiles.bat" goto Error_BatchFileInWrongLocation
 call "%~dp0Engine\Build\BatchFiles\GenerateProjectFiles.bat" %*
