Include *.ts files in the Git Repo.

#rnx

diff --git a/.gitignore b/.gitignore
--- a/.gitignore
+++ b/.gitignore
@@ -64,6 +64,7 @@
 !*.pl
 !*.py
 !*.js
+!*.ts
 !*.command
 
 # Other configuration and markup files
