Git: Include ISPC source files in the repository.

#rnx

diff --git a/.gitignore b/.gitignore
--- a/.gitignore
+++ b/.gitignore
@@ -35,6 +35,10 @@
 !*.exp
 !*.manifest
 
+# ISPC files
+!*.ispc
+!*.isph
+
 # Re-ignore vs host manifest
 *.vshost.exe.manifest
 
