PR #6509: [*] .gitignore actualized for UE 4.24.1 (Contributed by [REDACTED_PERSON])

#rnx

diff --git a/.gitignore b/.gitignore
--- a/.gitignore
+++ b/.gitignore
@@ -1,5 +1,5 @@
 #
-# This file tells Git about engine files that never really belong in source control.  They are usually build products, log 
+# This file tells Git about engine files that never really belong in source control.  They are usually build products, log
 # files and intermediate files generated from a compiler or the engine runtime.
 #
 #
@@ -11,7 +11,7 @@
 #  Paths in this file should use forward slash characters, not back slashes
 #  Use \ to escape special characters like ! and #
 #  Use ! to negate a previous pattern.  But it doesn't work if the parent sub-folder was masked out already.
-# 
+#
 
 # Ignore all files by default, but scan all directories
 *
@@ -167,14 +167,20 @@
 /Engine/Binaries/ThirdParty/Mono/Linux/bin/mono-boehm
 /Engine/Binaries/ThirdParty/Mono/Linux/bin/mono-sgen
 
-# ...and GitDependencies.exe
+# Add exceptions for GitDependencies.exe
 !/Engine/Binaries/DotNET/GitDependencies.exe
 !/Engine/Binaries/DotNET/GitDependencies.exe.config
 
 # Ignore documentation generated for C# tools
-Engine/Binaries/DotNET/UnrealBuildTool.xml
-Engine/Binaries/DotNET/AutomationScripts/BuildGraph.Automation.xml
-Engine/Source/Programs/AutomationTool/BuildGraph/BuildGraph.Automation.xml
+/Engine/Binaries/DotNET/UnrealBuildTool.xml
+/Engine/Binaries/DotNET/AutomationScripts/BuildGraph.Automation.xml
+/Engine/Binaries/DotNET/AutomationScripts/Newtonsoft.Json.xml
+/Engine/Binaries/DotNET/AutomationScripts/HoloLens/WindowsDevicePortalWrapper.xml
+/Engine/Source/Programs/AutomationTool/BuildGraph/BuildGraph.Automation.xml
+
+# Ignore Datasmith generated sources on all platforms
+/Engine/Binaries/**/DatasmithFacadeCSharp
+/Engine/Binaries/**/DatasmithSDK
 
 # Ignore version files in the Engine/Binaries directory created by UBT
 /Engine/Binaries/**/*.version
@@ -199,8 +205,8 @@ Saved/
 /LocalBuilds/
 
 # Ignore Swarm local save files
-Engine/Binaries/DotNET/SwarmAgent.DeveloperOptions.xml
-Engine/Binaries/DotNET/SwarmAgent.Options.xml
+/Engine/Binaries/DotNET/SwarmAgent.DeveloperOptions.xml
+/Engine/Binaries/DotNET/SwarmAgent.Options.xml
 
 # Explcitly ignore Mac DS_Store files, regardless of where they are
 .DS_Store
