PR #6622: Git Ignore files generated by [REDACTED_PERSON] (Contributed by [REDACTED_PERSON])

#rnx

diff --git a/.gitignore b/.gitignore
--- a/.gitignore
+++ b/.gitignore
@@ -185,8 +185,11 @@
 # Ignore version files in the Engine/Binaries directory created by UBT
 /Engine/Binaries/**/*.version
 
-# Ignore exp files in the the Engine/Binaries directory as they aren't C/C++ source files
-/Engine/Binaries/**/*.exp
+# Ignore exp files in Engine and Plugin Binaries directories as they aren't C/C++ source files
+/Engine/**/Binaries/**/*.exp
+
+# Ignore temporary code folders generated by live coding
+enc_temp_folder/
 
 # Ignore content folders
 Content/
