PR #6742: # Ignore documentation generated for C# tools (Contributed by [REDACTED_PERSON])

#rnx

diff --git a/.gitignore b/.gitignore
--- a/.gitignore
+++ b/.gitignore
@@ -175,6 +175,7 @@
 /Engine/Binaries/DotNET/UnrealBuildTool.xml
 /Engine/Binaries/DotNET/AutomationScripts/BuildGraph.Automation.xml
 /Engine/Binaries/DotNET/AutomationScripts/Newtonsoft.Json.xml
+/Engine/Binaries/DotNET/AutomationScripts/HoloLens/System.Runtime.WindowsRuntime.xml
 /Engine/Binaries/DotNET/AutomationScripts/HoloLens/WindowsDevicePortalWrapper.xml
 /Engine/Source/Programs/AutomationTool/BuildGraph/BuildGraph.Automation.xml
 
