Fix path to prereq installer.

diff --git a/Setup.bat b/Setup.bat
--- a/Setup.bat
+++ b/Setup.bat
@@ -24,7 +24,7 @@ echo Engine/Binaries/DotNET/GitDependencies.exe %* >>.git\hooks\post-merge
 
 rem Install prerequisites...
 echo Installing prerequisites...
-start /wait Engine\Extras\Redist\en-us\UE4PrereqSetup_x64.exe /quiet
+start /wait Engine\Extras\Redist\en-us\UEPrereqSetup_x64.exe /quiet
 
 rem Register the engine installation...
 if not exist .\Engine\Binaries\Win64\UnrealVersionSelector-Win64-Shipping.exe goto :no_unreal_version_selector
