Mask out some binaries from .gitignore.

diff --git a/.gitignore b/.gitignore
--- a/.gitignore
+++ b/.gitignore
@@ -171,6 +171,12 @@
 !/Engine/Binaries/DotNET/GitDependencies.exe
 !/Engine/Binaries/DotNET/GitDependencies.exe.config
 
+# Ignore NET binaries
+/Engine/Binaries/DotNET/UnrealBuildTool/
+/Engine/Binaries/DotNET/AutomationTool/
+/Engine/Source/Programs/DotNETCommon/BuildUtilities/bin/
+/Engine/Source/Programs/Shared/*/bin/
+
 # Ignore documentation generated for C# tools
 /Engine/Binaries/DotNET/UnrealBuildTool.xml
 /Engine/Binaries/DotNET/AutomationScripts/BuildGraph.Automation.xml
