Exclude UE5 codelite files from Git.

diff --git a/.gitignore b/.gitignore
--- a/.gitignore
+++ b/.gitignore
@@ -130,6 +130,8 @@
 /.kdev4/
 /UE4CodeCompletionFolders.txt
 /UE4CodeLitePreProcessor.txt
+/UE5CodeCompletionFolders.txt
+/UE5CodeLitePreProcessor.txt
 
 # Add exceptions for Mac Mono binaries needed to run GitDependencies.exe.
 !/Engine/Binaries/ThirdParty/Mono/Mac/bin/*
