Update Linux github doc

Remove old reference to v11 cross toolchain
Move to different wiki URL location

PR #7279

diff --git a/README.md b/README.md
--- a/README.md
+++ b/README.md
@@ -97,9 +97,9 @@ The steps below will take you through cloning your own private fork, then compil
 
 1. Both cross-compiling and native builds are supported. 
 
-   **Cross-compiling** is handy when you are a Windows (Mac support planned too) developer who wants to package your game for Linux with minimal hassle, and it requires a [cross-compiler toolchain](http://cdn.unrealengine.com/CrossToolchain_Linux/v11_clang-5.0.0-centos7.zip) to be installed (see the [Linux cross-compiling page on the wiki](https://docs.unrealengine.com/latest/INT/Platforms/Linux/GettingStarted/)).
+   **Cross-compiling** is handy when you are a Windows (Mac support planned too) developer who wants to package your game for Linux with minimal hassle, and it requires a cross-compiler toolchain to be installed (see the [Linux cross-compiling page on the wiki](https://docs.unrealengine.com/latest/INT/Platforms/Linux/GettingStarted/)).
 
-   **Native compilation** is discussed in [a separate README](Engine/Build/BatchFiles/Linux/README.md) and [community wiki page](https://www.ue4community.wiki/Legacy/Building_On_Linux). 
+   **Native compilation** is discussed in [a separate README](Engine/Build/BatchFiles/Linux/README.md) and [community wiki page](https://unrealcommunity.wiki/building-on-linux-qr8t0si2).
 
 
 
