Add json test data exclusions for NeuralNetwork

diff --git a/.gitignore b/.gitignore
--- a/.gitignore
+++ b/.gitignore
@@ -241,3 +241,6 @@ Engine/Extras/ThirdPartyNotUE/SDKs/HostLinux/**
 
 # Ignore USD resources copied into Binaries directory at build time
 /Engine/Binaries/ThirdParty/USD/**
+
+# Ignore test json data files in NeuralNetworkInterface
+/Engine/Plugins/Experimental/NeuralNetworkInference/Source/ThirdParty/Deps/json/**
