Replaced UE4 ocurrences with UE5 in GIT README.md (Contributed by [REDACTED_PERSON])
Replaced Visual Studio 2017 occurrences of with Visual Studio 2019

#8527

diff --git a/README.md b/README.md
--- a/README.md
+++ b/README.md
@@ -44,9 +44,9 @@ The steps below will take you through cloning your own private fork, then compil
    If you'd prefer not to use Git, you can get the source with the 'Download ZIP' button on the right. The built-in Windows zip utility will mark the contents of zip files 
    downloaded from the Internet as unsafe to execute, so right-click the zip file and select 'Properties...' and 'Unblock' before decompressing it. Third-party zip utilities don't normally do this.
 
-1. Install **Visual Studio 2017**. 
-   All desktop editions of Visual Studio 2017 can build UE4, including [Visual Studio Community 2017](http://www.visualstudio.com/products/visual-studio-community-vs), which is free for small teams and individual developers.
-   To install the correct components for UE4 development, check the "Game Development with C++" workload, and the "Unreal Engine Installer" and "Nuget Package Manager" optional components.
+1. Install **Visual Studio 2019**. 
+   All desktop editions of Visual Studio 2019 can build UE5, including [Visual Studio Community 2019](http://www.visualstudio.com/products/visual-studio-community-vs), which is free for small teams and individual developers.
+   To install the correct components for UE5 development, check the "Game Development with C++" workload, and the "Unreal Engine Installer" and "Nuget Package Manager" optional components.
   
 1. Open your source folder in Explorer and run **Setup.bat**. 
    This will download binary content for the engine, as well as installing prerequisites and setting up Unreal file associations. 
@@ -57,10 +57,10 @@ The steps below will take you through cloning your own private fork, then compil
  
 1. Run **GenerateProjectFiles.bat** to create project files for the engine. It should take less than a minute to complete.  
 
-1. Load the project into Visual Studio by double-clicking on the **UE4.sln** file. Set your solution configuration to **Development Editor** and your solution
-   platform to **Win64**, then right click on the **UE4** target and select **Build**. It may take anywhere between 10 and 40 minutes to finish compiling, depending on your system specs.
+1. Load the project into Visual Studio by double-clicking on the **UE5.sln** file. Set your solution configuration to **Development Editor** and your solution
+   platform to **Win64**, then right click on the **UE5** target and select **Build**. It may take anywhere between 10 and 40 minutes to finish compiling, depending on your system specs.
 
-1. After compiling finishes, you can load the editor from Visual Studio by setting your startup project to **UE4** and pressing **F5** to debug.
+1. After compiling finishes, you can load the editor from Visual Studio by setting your startup project to **UE5** and pressing **F5** to debug.
 
 
 
@@ -80,8 +80,8 @@ The steps below will take you through cloning your own private fork, then compil
 
 1. In the same folder, double-click **GenerateProjectFiles.command**.  It should take less than a minute to complete.  
 
-1. Load the project into Xcode by double-clicking on the **UE4.xcworkspace** file. Select the **ShaderCompileWorker** for **My Mac** target in the title bar,
-   then select the 'Product > Build' menu item. When Xcode finishes building, do the same for the **UE4** for **My Mac** target. Compiling may take anywhere between 15 and 40 minutes, depending on your system specs.
+1. Load the project into Xcode by double-clicking on the **UE5.xcworkspace** file. Select the **ShaderCompileWorker** for **My Mac** target in the title bar,
+   then select the 'Product > Build' menu item. When Xcode finishes building, do the same for the **UE5** for **My Mac** target. Compiling may take anywhere between 15 and 40 minutes, depending on your system specs.
    
 1. After compiling finishes, select the 'Product > Run' menu item to load the editor.
 
