Update GitHub readme for UE5, new branch descriptions, refresh console and platform instructions, refresh licensing section, and more

diff --git a/README.md b/README.md
--- a/README.md
+++ b/README.md
@@ -1,19 +1,20 @@
 Unreal Engine
 =============
 
-Welcome to the Unreal Engine source code! 
+Welcome to the Unreal Engine source code!
 
-From this repository you can build the Unreal Editor for Windows, Mac and Linux, compile Unreal Engine games for Android, iOS, PlayStation 4, Xbox One and HTML5,
-and build tools like Unreal Lightmass and Unreal Frontend. Modify them in any way you can imagine, and share your changes with others! 
+With the code in this repository, you can build the Unreal Editor for Windows, Mac, and Linux; compile Unreal Engine games for a variety of target platforms, including desktop, consoles, mobile, and embedded devices; and build tools like Unreal Lightmass and Unreal Frontend. Modify the code in any way you can imagine, and share your changes with others!
 
-We have a heap of documentation available for the engine on the web. If you're looking for the answer to something, you may want to start here: 
+We have a vast amount of [official documentation](http://docs.unrealengine.com) available for the engine. If you're looking for the answer to something, you may want to start in one of these places:
 
-* [Unreal Engine Programming Guide](https://docs.unrealengine.com/latest/INT/Programming/index.html)
-* [Unreal Engine API Reference](https://docs.unrealengine.com/latest/INT/API/index.html)
-* [Engine source and GitHub on the Unreal Engine forums](https://forums.unrealengine.com/forumdisplay.php?1-Development-Discussion)
+*   [Programming and Scripting in Unreal Engine](https://docs.unrealengine.com/ProgrammingAndScripting)
+*   [Development Setup](https://docs.unrealengine.com/ProductionPipelines/DevelopmentSetup/)
+*   [Working with the GitHub source code distribution](https://docs.unrealengine.com/ProgrammingAndScripting/ProgrammingWithCPP/DownloadingSourceCode)
+*   [Unreal Engine C++ API Reference](https://docs.unrealengine.com/API)
+*   [Community forum about working with engine source and GitHub](https://forums.unrealengine.com/c/development-discussion/c-programming/engine-source-github/33)
 
-If you need more, just ask! A lot of Epic developers hang out on the [forums](https://forums.unrealengine.com/) or [AnswerHub](https://answers.unrealengine.com/), 
-and we're proud to be part of a well-meaning, friendly and welcoming community of thousands. 
+If you need more, just ask! Many Epic developers read the [forums](https://forums.unrealengine.com/) and [AnswerHub](https://answers.unrealengine.com/),
+and we're proud to be part of a well-meaning, friendly, and welcoming community of thousands.
 
 
 Branches
@@ -21,12 +22,17 @@ Branches
 
 We publish source for the engine in several branches:
 
-The **[release branch](https://github.com/EpicGames/UnrealEngine/tree/release)** is extensively tested by our QA team and makes a great starting point for learning the engine or
-making your own games. We work hard to make releases stable and reliable, and aim to publish new releases every few months.
+*   Numbered branches identify past and upcoming official releases, and the **[release](https://github.com/EpicGames/UnrealEngine/tree/release)** branch always reflects the current official release. These are extensively tested by our QA team, so they make a great starting point for learning Unreal Engine and for making your own games. We work hard to make releases stable and reliable, and aim to publish a new release every few months.
 
-The **[master branch](https://github.com/EpicGames/UnrealEngine/tree/master)** is the hub of changes from all our specialized engine development teams. Our internal game teams typically take engine snapshots from here, but it isn't subject to as much testing as release branches.
+*   Most active development on UE5 happens in the **[ue5-main](https://github.com/EpicGames/UnrealEngine/tree/ue5-main)** branch. This branch reflects the cutting edge of the engine and may be buggy — it may not even compile. We make it available for battle-hardened developers eager to test new features or work in lock-step with us.
 
-Individual teams have their own **development branches** for day to day work ([dev-core](https://github.com/EpicGames/UnrealEngine/tree/dev-core), [dev-mobile](https://github.com/EpicGames/UnrealEngine/tree/dev-mobile) and [dev-sequencer](https://github.com/EpicGames/UnrealEngine/tree/dev-sequencer), for example). These branches reflect the cutting edge of the engine and may be buggy - they may not even compile. Battle-hardened developers eager to test new features or work lock-step with us should head to one of these. We aim to merge development branches to master every 3-4 weeks.
+    If you choose to work in this branch, be aware that it is likely to be **ahead** of the branches for the current official release and the next upcoming release. Therefore, content and code that you create to work with the ue5-main branch may not be compatible with public releases until we create a new branch directly from ue5-main for a future official release.
+
+*   For a stable Early Access release of UE5, use the **[ue5-early-access](https://github.com/EpicGames/UnrealEngine/tree/ue5-early-access)** branch. The code in this branch matches the UE5 Early Access binary release that is available in the Epic Games Launcher. It’s a good place to start if you want a solid foundation for trying out the new features in UE5. For more information, see [the UE5 documentation](https://docs.unrealengine.com/5.0/en-US/Welcome/).
+
+*   The **[master branch](https://github.com/EpicGames/UnrealEngine/tree/master)** is the hub of changes to UE4 from all our engine development teams. It’s not subject to as much testing as release branches, so if you’re using UE4 this may not be the best choice for you.
+
+*   Branches whose names contain `dev`, `staging`, and `test` are typically for internal Epic processes, and are rarely useful for end users.
 
 Other short-lived branches may pop-up from time to time as we stabilize new releases or hotfixes.
 
@@ -34,101 +40,113 @@ Other short-lived branches may pop-up from time to time as we stabilize new rele
 Getting up and running
 ----------------------
 
-The steps below will take you through cloning your own private fork, then compiling and running the editor yourself:
+The steps below take you through cloning your own private fork, then compiling and running the editor yourself:
 
 ### Windows
 
-1. Install **[GitHub for Windows](https://windows.github.com/)** then **[fork and clone our repository](https://guides.github.com/activities/forking/)**. 
-   To use Git from the command line, see the [Setting up Git](https://help.github.com/articles/set-up-git/) and [Fork a Repo](https://help.github.com/articles/fork-a-repo/) articles.
+1.  Install a Git client like **[GitHub for Windows](https://windows.github.com/)**, then **[fork and clone our repository](https://guides.github.com/activities/forking/)**.
+
+    To use Git from the command line instead, see the [Setting up Git](https://help.github.com/articles/set-up-git/) and [Fork a Repo](https://help.github.com/articles/fork-a-repo/) articles.
+
+    If you'd prefer not to use Git, you can get the source with the **Download ZIP** button on the right. Note that the zip utility built in to Windows marks the contents of .zip files downloaded from the Internet as unsafe to execute, so right-click the .zip file and select **Properties…** and **Unblock** before decompressing it.
+
+1.  Install **Visual Studio 2019**.
+
+    All desktop editions of Visual Studio 2019 can build UE5, including [Visual Studio Community 2019](http://www.visualstudio.com/products/visual-studio-community-vs), which is free for small teams and individual developers.
+
+    To install the correct components for UE5 development, make sure the **Game Development with C++** workload is checked. Under the **Installation Details** section on the right, also choose the following components:
+
+    -   **C++ profiling tools**
+    -   **C++ AddressSanitizer** (optional)
+    -   **Windows 10 SDK** (10.0.18362 or newer)
+    -   **Unreal Engine Installer**
 
-   If you'd prefer not to use Git, you can get the source with the 'Download ZIP' button on the right. The built-in Windows zip utility will mark the contents of zip files 
-   downloaded from the Internet as unsafe to execute, so right-click the zip file and select 'Properties...' and 'Unblock' before decompressing it. Third-party zip utilities don't normally do this.
+1.  Open your source folder in Windows Explorer and run **Setup.bat**. This will download binary content for the engine, install prerequisites, and set up Unreal file associations.
 
-1. Install **Visual Studio 2019**. 
-   All desktop editions of Visual Studio 2019 can build UE5, including [Visual Studio Community 2019](http://www.visualstudio.com/products/visual-studio-community-vs), which is free for small teams and individual developers.
-   To install the correct components for UE5 development, check the "Game Development with C++" workload, and the "Unreal Engine Installer" and "Nuget Package Manager" optional components.
-  
-1. Open your source folder in Explorer and run **Setup.bat**. 
-   This will download binary content for the engine, as well as installing prerequisites and setting up Unreal file associations. 
-   On Windows 8, a warning from SmartScreen may appear.  Click "More info", then "Run anyway" to continue.
-   
-   A clean download of the engine binaries is currently 3-4gb, which may take some time to complete.
-   Subsequent checkouts only require incremental downloads and will be much quicker.
- 
-1. Run **GenerateProjectFiles.bat** to create project files for the engine. It should take less than a minute to complete.  
+    On Windows 8, a warning from SmartScreen may appear. Click **More info**, then **Run anyway** to continue.
 
-1. Load the project into Visual Studio by double-clicking on the **UE5.sln** file. Set your solution configuration to **Development Editor** and your solution
-   platform to **Win64**, then right click on the **UE5** target and select **Build**. It may take anywhere between 10 and 40 minutes to finish compiling, depending on your system specs.
+    A clean download of the engine binaries is currently 3-4 GB, which may take some time to complete. Subsequent runs will be much faster, as they only download new and updated content.
 
-1. After compiling finishes, you can load the editor from Visual Studio by setting your startup project to **UE5** and pressing **F5** to debug.
+1.  Run **GenerateProjectFiles.bat** to create project files for the engine. It should take less than a minute to complete.  
 
+1.  Load the project into Visual Studio by double-clicking the new **UE5.sln** file.
 
+1.  Set your solution configuration to **Development Editor** and your solution platform to **Win64**, then right click the **UE5** target and select **Build**. It may take anywhere between 10 and 40 minutes to finish compiling, depending on your system specs.
+
+1.  After compiling finishes, you can run the editor from Visual Studio by setting your startup project to **UE5** and pressing **F5** to start debugging.
 
 
 ### Mac
-   
-1. Install **[GitHub for Mac](https://mac.github.com/)** then **[fork and clone our repository](https://guides.github.com/activities/forking/)**. 
-   To use Git from the Terminal, see the [Setting up Git](https://help.github.com/articles/set-up-git/) and [Fork a Repo](https://help.github.com/articles/fork-a-repo/) articles.
-   If you'd rather not use Git, use the 'Download ZIP' button on the right to get the source directly.
 
-1. Install the latest version of [Xcode](https://itunes.apple.com/us/app/xcode/id497799835).
+1.  Install a Git client like **[GitHub for Mac](https://mac.github.com/)**, then **[fork and clone our repository](https://guides.github.com/activities/forking/)**.
+
+    To use Git from the Terminal instead, see the [Setting up Git](https://help.github.com/articles/set-up-git/) and [Fork a Repo](https://help.github.com/articles/fork-a-repo/) articles.
+
+    If you'd rather not use Git, use the **Download ZIP** button on the right to get the source directly.
 
-1. Open your source folder in Finder and double-click on **Setup.command** to download binary content for the engine. You can close the Terminal window afterwards.
+1.  Install the latest version of [Xcode](https://itunes.apple.com/us/app/xcode/id497799835).
 
-   If you downloaded the source as a .zip file, you may see a warning about it being from an unidentified developer (because .zip files on GitHub aren't digitally signed).
-   To work around it, right-click on Setup.command, select Open, then click the Open button.
+1.  Open your source folder in Finder and double-click **Setup.command** to download binary content for the engine. You can close the Terminal window afterwards.
 
-1. In the same folder, double-click **GenerateProjectFiles.command**.  It should take less than a minute to complete.  
+    If you downloaded the source as a .zip file, you may see a warning about it being from an unidentified developer, because .zip files on GitHub aren't digitally signed. To work around this, right-click **Setup.command**, select **Open**, then click the **Open** button.
 
-1. Load the project into Xcode by double-clicking on the **UE5.xcworkspace** file. Select the **ShaderCompileWorker** for **My Mac** target in the title bar,
-   then select the 'Product > Build' menu item. When Xcode finishes building, do the same for the **UE5** for **My Mac** target. Compiling may take anywhere between 15 and 40 minutes, depending on your system specs.
-   
-1. After compiling finishes, select the 'Product > Run' menu item to load the editor.
+1.  In the same folder, double-click **GenerateProjectFiles.command**. It should take less than a minute to complete.  
 
+1.  Load the project into Xcode by double-clicking the **UE5.xcworkspace** file. Select the **ShaderCompileWorker** for **My Mac** target in the title bar, then select the **Product > Build** menu item. When Xcode finishes building, do the same for the **UE5** for **My Mac** target. Compiling may take anywhere between 15 and 40 minutes, depending on your system specs.
 
+1.  After compiling finishes, select the **Product > Run** menu item to load the editor.
 
 
 ### Linux
 
-1. [Set up Git](https://help.github.com/articles/set-up-git/) and [fork our repository](https://help.github.com/articles/fork-a-repo/).
-   If you'd prefer not to use Git, use the 'Download ZIP' button on the right to get the source as a zip file.
+1.  Install a **[visual Git client](https://git-scm.com/download/gui/linux)**, then **[fork and clone our repository](https://guides.github.com/activities/forking/)**.
 
-1. Open your source folder and run **Setup.sh** to download binary content for the engine.
+    To use Git from the command line instead, see the [Setting up Git](https://help.github.com/articles/set-up-git/) and [Fork a Repo](https://help.github.com/articles/fork-a-repo/) articles.
 
-1. Both cross-compiling and native builds are supported. 
+    If you'd prefer not to use Git, use the **Download ZIP** button on the right to get the source as a zip file.
 
-   **Cross-compiling** is handy when you are a Windows (Mac support planned too) developer who wants to package your game for Linux with minimal hassle, and it requires a cross-compiler toolchain to be installed (see the [Linux cross-compiling page on the wiki](https://docs.unrealengine.com/latest/INT/Platforms/Linux/GettingStarted/)).
+1.  Open your source folder and run **Setup.sh** to download binary content for the engine.
 
-   **Native compilation** is discussed in [a separate README](Engine/Build/BatchFiles/Linux/README.md) and [community wiki page](https://unrealcommunity.wiki/building-on-linux-qr8t0si2).
+1.  Both cross-compiling and native builds are supported.
 
+    -   **Cross-compiling** is handy for Windows developers who want to package a game for Linux with minimal hassle. It requires a cross-compiler toolchain to be installed. See the [Linux cross-compiling page in the documentation](https://docs.unrealengine.com/SharingAndReleasing/Linux/GettingStarted/).
 
+    -   **Native compilation** is discussed in [a separate README](Engine/Build/BatchFiles/Linux/README.md) and [community wiki page](https://unrealcommunity.wiki/building-on-linux-qr8t0si2).
 
 
 ### Additional target platforms
 
-**Android** support will be downloaded by the setup script if you have the Android NDK installed. See the [Android getting started guide](https://docs.unrealengine.com/latest/INT/Platforms/Android/GettingStarted/).
+*   **Android** support will be downloaded by the setup script if you have the Android NDK installed. See the [Android Quick Start guide](https://docs.unrealengine.com/en-US/SharingAndReleasing/Mobile/Android/GettingStarted/).
+
+*   **iOS** development requires a Mac. Instructions are in the [iOS Quick Start guide](https://docs.unrealengine.com/en-US/SharingAndReleasing/Mobile/iOS/QuickStart/).
+
+*   Development for consoles and other platforms with restricted access, like **Sony PlayStation**, **Microsoft Xbox**, and **Nintendo Switch**, is only possible if you have a registered developer account with those third-party vendors.
 
-**iOS** programming requires a Mac. Instructions are in the [iOS getting started guide](https://docs.unrealengine.com/latest/INT/Platforms/iOS/GettingStarted/index.html).
+    Depending on the platform, additional documentation or guidance may be available in the [Unreal Developer Network support site](https://udn.unrealengine.com/s/), or as a downloadable archive in the section of the [Unreal Engine Forums](https://forums.unrealengine.com/) that is dedicated to your platform.
 
-**HTML5** support will be downloaded by the setup script if you have Emscripten installed. Please see the [HTML5 getting started guide](https://docs.unrealengine.com/latest/INT/Platforms/HTML5/GettingStarted/index.html).
+    If you don’t have access to these resources, first register a developer account with the third party vendor. Then contact your Epic Games account manager if you have one, or fill out and submit the [Console Development Request form](https://epicgames.secure.force.com/Forms/FormConsoleAccessRequest?) for Unreal Engine if you don’t. Epic will contact you with a formal agreement to digitally sign. Once this is approved, you will receive instructions on how to access source code, binaries, and additional instructions for your platform.
 
-**PlayStation 4** or **Xbox One** development requires additional files that can only be provided after your registered developer status is confirmed by Sony or Microsoft. See [the announcement blog post](https://www.unrealengine.com/blog/playstation-4-and-xbox-one-now-supported) for more information.
 
+Licensing
+---------
 
-Licensing and Contributions
----------------------------
+Your access to and use of Unreal Engine on GitHub is governed by an End User License Agreement (EULA). You can choose either a Publishing EULA or a Creators EULA, depending on how the terms fit your intended usage of the engine. For a comparison and links to FAQs, see the [Unreal Engine download page](https://www.unrealengine.com/en-US/download). If you don't agree to the terms in your chosen EULA, as amended from time to time, you are not permitted to access or use Unreal Engine.
 
-Your access to and use of Unreal Engine on GitHub is governed by the [Unreal Engine End User License Agreement](https://www.unrealengine.com/eula). If you don't agree to those terms, as amended from time to time, you are not permitted to access or use Unreal Engine.
+Contributions
+-------------
 
-We welcome any contributions to Unreal Engine development through [pull requests](https://github.com/EpicGames/UnrealEngine/pulls/) on GitHub. Most of our active development is in the **master** branch, so we prefer to take pull requests there (particularly for new features). We try to make sure that all new code adheres to the [Epic coding standards](https://docs.unrealengine.com/latest/INT/Programming/Development/CodingStandard/).  All contributions are governed by the terms of the EULA.
+We welcome contributions to Unreal Engine development through [pull requests](https://github.com/EpicGames/UnrealEngine/pulls/) on GitHub.
+
+We prefer to take pull requests in our active development branches, particularly for new features. For UE5, use the **ue5-main** branch; for UE4, use the **master** branch. Please make sure that all new code adheres to the [Epic coding standards](https://docs.unrealengine.com/ProductionPipelines/DevelopmentSetup/CodingStandard).
+
+For more information on the process and expectations, see [the documentation](https://docs.unrealengine.com/ProgrammingAndScripting/ProgrammingWithCPP/DownloadingSourceCode/ContributingToUnreal/).
+
+All contributions are governed by the terms of your EULA.
 
 
 Additional Notes
 ----------------
 
-The first time you start the editor from a fresh source build, you may experience long load times. 
-The engine is optimizing content for your platform to the _derived data cache_, and it should only happen once.
-
-Your private forks of the Unreal Engine code are associated with your GitHub account permissions.
-If you unsubscribe or switch GitHub user names, you'll need to re-fork and upload your changes from a local copy. 
+The first time you start the editor from a fresh source build, you may experience long load times. The engine is optimizing content for your platform and storing it in the _[derived data cache](https://docs.unrealengine.com/ProductionPipelines/DerivedDataCache/)_. This should only happen once.
 
+Your private forks of the Unreal Engine code are associated with your GitHub account permissions. If you unsubscribe or switch GitHub user names, you'll need to create a new fork and upload your changes from the fresh copy.
