Remove console specific .natvis, it's not necessary anymore. It also wasn't being used because the paths were broken.
All .natvis types come from the cross platform Unreal.natvis

diff --git a/GenerateProjectFiles.bat b/GenerateProjectFiles.bat
--- a/GenerateProjectFiles.bat
+++ b/GenerateProjectFiles.bat
@@ -1,10 +1,5 @@
 @echo off
 
-REM Install PS4 visualizer if the SDK and installation file are present
-if exist "%~dp0Engine\Extras\VisualStudioDebugging\PS4\InstallPS4Visualizer.bat" (
-  call "%~dp0Engine\Extras\VisualStudioDebugging\PS4\InstallPS4Visualizer.bat"
-)
-
 if not exist "%~dp0Engine\Build\BatchFiles\GenerateProjectFiles.bat" goto Error_BatchFileInWrongLocation
 call "%~dp0Engine\Build\BatchFiles\GenerateProjectFiles.bat" %*
 exit /B %ERRORLEVEL%
