Adding VP Link vcam to GH - adding .gitignore exclusion rules

#rnx
#okforgithhub public

diff --git a/.gitignore b/.gitignore
--- a/.gitignore
+++ b/.gitignore
@@ -62,6 +62,9 @@
 !*.txt
 !*.md
 
+# Swift files
+*.swift
+
 # Script files
 !*.bat
 !*.sh
@@ -87,6 +90,7 @@
 !*.version
 !.git*
 !.editorconfig
+!.plist
 
 # Projects and makefiles
 !*.cmake
@@ -146,6 +150,9 @@
 !/Engine/Binaries/ThirdParty/Mono/Mac/lib/mono/gac/System.Core/4.0.0.0__b77a5c561934e089/System.Core.dll
 !/Engine/Binaries/ThirdParty/Mono/Mac/lib/mono/gac/System.Xml/4.0.0.0__b77a5c561934e089/System.Xml.dll
 
+# Exception for VirtualProduction LiveLinkVCamFiles
+!/Engine/Extras/VirtualProduction/LiveLinkVCAM/**
+
 # Ignore the automatically created mono symlink
 /Engine/Binaries/ThirdParty/Mono/Mac/bin/mono
 
