Remove dotnet from the repo, will be added to the s3 bucket, and ensure GitDependencies is in the repo for all platforms

#rnx

diff --git a/.gitignore b/.gitignore
--- a/.gitignore
+++ b/.gitignore
@@ -176,11 +176,11 @@
 /Engine/Binaries/ThirdParty/Mono/Linux/bin/mono-boehm
 /Engine/Binaries/ThirdParty/Mono/Linux/bin/mono-sgen
 
-# Add exceptions for GitDependencies.exe
-!/Engine/Binaries/DotNET/GitDependencies.exe
-!/Engine/Binaries/DotNET/GitDependencies.exe.config
+# Add exceptions for GitDependencies
+!/Engine/Binaries/DotNET/GitDependencies/
 
 # Ignore NET binaries
+/Engine/Binaries/ThirdParty/DotNet/
 /Engine/Binaries/DotNET/UnrealBuildTool/
 /Engine/Binaries/DotNET/AutomationTool/
 /Engine/Source/Programs/DotNETCommon/BuildUtilities/bin/
