Add ignore entries for two LiveLinkCam files

They are too large to be in the repo

#rnx

diff --git a/.gitignore b/.gitignore
--- a/.gitignore
+++ b/.gitignore
@@ -152,6 +152,8 @@
 
 # Exception for VirtualProduction LiveLinkVCamFiles
 !/Engine/Extras/VirtualProduction/LiveLinkVCAM/**
+/Engine/Extras/VirtualProduction/LiveLinkVCAM/UE4LiveLinkProvider.framework/UE4LiveLinkProvider
+/Engine/Extras/VirtualProduction/LiveLinkVCAM/UE5LiveLinkProvider.framework/UE5LiveLinkProvider
 
 # Ignore the automatically created mono symlink
 /Engine/Binaries/ThirdParty/Mono/Mac/bin/mono
