Follow up to 19099995

Rename UE4Games.uprojectdirs to Default.uprojectdirs.
New naming is meant to reflect (a) that we're not UE4 any more, and (b) that these are the default locations to check for .uproject files rather than anything game-specific. Anyone with this file open in UGS may need to reopen the project with the renamed file.

#rnx

diff --git a/Default.uprojectdirs b/Default.uprojectdirs
new file mode 100644
--- /dev/null
+++ b/Default.uprojectdirs
@@ -0,0 +1,9 @@
+; These folders will be searched 1 level deep in order to find projects
+; UnrealBuildTool will store the following information:
+; - Project name
+; - Location of project
+; - Whether it has code or not
+; - TargetNames contains at the project location
+;
+./
+Engine/Source/
diff --git a/UE4Games.uprojectdirs b/UE4Games.uprojectdirs
deleted file mode 100644
