Updating .gitignore for github after some recent renames

Validated with a local copy of the git repo and the updated ignore file

#rnx

diff --git a/.gitignore b/.gitignore
--- a/.gitignore
+++ b/.gitignore
@@ -252,4 +252,4 @@ Engine/Extras/ThirdPartyNotUE/SDKs/HostLinux/**
 /Engine/Binaries/ThirdParty/USD/**
 
 # Ignore test json data files in NeuralNetworkInterface
-/Engine/Plugins/Experimental/NeuralNetworkInference/Source/ThirdParty/Deps/json/**
+/Engine/Plugins/Experimental/NNI/Source/ThirdParty/Deps/json/**
