Accepting GitHub #8958 pull request to add --force and --prompt options to Setup on GitHub build

diff --git a/Setup.sh b/Setup.sh
--- a/Setup.sh
+++ b/Setup.sh
@@ -11,6 +11,33 @@ if [ ! -f Engine/Binaries/DotNET/GitDependencies.exe ]; then
 	exit 1
 fi 
 
+function FindPromptOrForceArg()
+{
+	read -r -a Args <<< "$@"
+	for Arg in "${Args[@]}"
+	do
+		case $Arg in
+				--prompt)
+					echo "${Arg}"
+					return
+				;;
+				--force)
+					echo "${Arg}"
+					return
+				;;
+		esac
+	done
+}
+
+declare -a GitDepsArgs
+PromptOrForce=$(FindPromptOrForceArg "$@")
+if [ -z "${PromptOrForce}" ];
+then
+	# Use --prompt by default, whenever caller hasn't supplied a --prompt or --force arg
+  GitDepsArgs=(--prompt)
+fi
+GitDepsArgs+=("$@")
+
 if [ "$(uname)" = "Darwin" ]; then
 	# Setup the git hooks
 	if [ -d .git/hooks ]; then
@@ -22,7 +49,7 @@ if [ "$(uname)" = "Darwin" ]; then
 	fi
 
 	# Get the dependencies for the first time
-	Engine/Build/BatchFiles/Mac/GitDependencies.sh --prompt $@
+	Engine/Build/BatchFiles/Mac/GitDependencies.sh "${GitDepsArgs[@]}"
 else
 	# Setup the git hooks
 	if [ -d .git/hooks ]; then
@@ -37,7 +64,7 @@ else
 	fi
 
 	# Get the dependencies for the first time
-	Engine/Build/BatchFiles/Linux/GitDependencies.sh --prompt $@
+	Engine/Build/BatchFiles/Linux/GitDependencies.sh "${GitDepsArgs[@]}"
 
 	echo Register the engine installation...
 	if [ -f Engine/Binaries/Linux/UnrealVersionSelector-Linux-Shipping ]; then
