Fix GenerateProjectFiles.sh to check for .uedependencies and .ue4dependencies

#rnx

diff --git a/GenerateProjectFiles.sh b/GenerateProjectFiles.sh
--- a/GenerateProjectFiles.sh
+++ b/GenerateProjectFiles.sh
@@ -12,7 +12,7 @@ if [ ! -f Engine/Build/BatchFiles/Mac/GenerateProjectFiles.sh ]; then
 fi 
 
 if [ -f Setup.sh ]; then
-	if [ ! -f .ue4dependencies ]; then
+	if [ ! -f .uedependencies ] && [ ! -f .ue4dependencies ]; then
 		echo "Please run Setup to download dependencies before generating project files."
 		exit 1
 	fi
