Adding some paths to exclude from git

#rnx

diff --git a/.gitignore b/.gitignore
--- a/.gitignore
+++ b/.gitignore
@@ -196,6 +196,15 @@
 /Engine/Binaries/DotNET/AutomationScripts/HoloLens/WindowsDevicePortalWrapper.xml
 /Engine/Source/Programs/AutomationTool/BuildGraph/BuildGraph.Automation.xml
 
+# Ignore generated files from C# tools
+/Engine/Source/Programs/**/*.deps.json
+/Engine/Source/Programs/**/*.runtimeconfig*.json
+/Engine/Source/Programs/**/Properties/launchSettings.json
+/Engine/Source/Programs/UnrealBuildTool.Tests/bin/
+/Engine/Binaries/DotNet/**/*.deps.json
+/Engine/Binaries/DotNet/Android/UnrealAndroidFileTool/*.xml
+/Engine/Binaries/DotNet/**/*.runtimeconfig*.json
+
 # Ignore Datasmith generated sources on all platforms
 /Engine/Binaries/**/DatasmithFacadeCSharp
 /Engine/Binaries/**/DatasmithSDK
