Add new exception for LiveLinkProvider

#rnx

diff --git a/.gitignore b/.gitignore
--- a/.gitignore
+++ b/.gitignore
@@ -152,6 +152,7 @@
 
 # Exception for VirtualProduction LiveLinkVCamFiles
 !/Engine/Extras/VirtualProduction/LiveLinkVCAM/**
+/Engine/Extras/VirtualProduction/LiveLinkVCAM/LiveLinkProvider.framework/LiveLinkProvider
 /Engine/Extras/VirtualProduction/LiveLinkVCAM/UE4LiveLinkProvider.framework/UE4LiveLinkProvider
 /Engine/Extras/VirtualProduction/LiveLinkVCAM/UE5LiveLinkProvider.framework/UE5LiveLinkProvider
 
