PR #9101 Fixing references to UE4

#rnx

diff --git a/GenerateProjectFiles.bat b/GenerateProjectFiles.bat
--- a/GenerateProjectFiles.bat
+++ b/GenerateProjectFiles.bat
@@ -5,6 +5,6 @@ call "%~dp0Engine\Build\BatchFiles\GenerateProjectFiles.bat" %*
 exit /B %ERRORLEVEL%
 
 :Error_BatchFileInWrongLocation
-echo GenerateProjectFiles ERROR: The batch file does not appear to be located in the root UE4 directory.  This script must be run from within that directory.
+echo GenerateProjectFiles ERROR: The batch file does not appear to be located in the root UE directory.  This script must be run from within that directory.
 pause
 exit /B 1
diff --git a/Setup.sh b/Setup.sh
--- a/Setup.sh
+++ b/Setup.sh
@@ -7,7 +7,7 @@ cd "`dirname "$0"`"
 
 if [ ! -f Engine/Binaries/DotNET/GitDependencies.exe ]; then
 	echo "GitSetup ERROR: This script does not appear to be located \
-       in the root UE4 directory and must be run from there."
+       in the root UE directory and must be run from there."
 	exit 1
 fi 
 
