Reverting to the platform specific GitDependencies

#rnx

diff --git a/Engine/Build/BatchFiles/Linux/GitDependencies.sh b/Engine/Build/BatchFiles/Linux/GitDependencies.sh
--- a/Engine/Build/BatchFiles/Linux/GitDependencies.sh
+++ b/Engine/Build/BatchFiles/Linux/GitDependencies.sh
@@ -20,7 +20,7 @@ cd ../../../..
 RESULT=0
 
 while : ; do
-        dotnet ./Engine/Binaries/DotNET/GitDependencies.dll $ARGS
+        ./Engine/Binaries/DotNET/GitDependencies/linux-x64/GitDependencies $ARGS
         RESULT=$?
 
         echo "Result: $RESULT"
diff --git a/Engine/Build/BatchFiles/Mac/GitDependencies.sh b/Engine/Build/BatchFiles/Mac/GitDependencies.sh
--- a/Engine/Build/BatchFiles/Mac/GitDependencies.sh
+++ b/Engine/Build/BatchFiles/Mac/GitDependencies.sh
@@ -10,14 +10,14 @@ fi
 
 cd "$(dirname "$SCRIPT_PATH")" && SCRIPT_PATH="`pwd`/$(basename "$SCRIPT_PATH")"
 
-if [ ! -f ../../../Binaries/DotNET/GitDependencies.dll ]; then
+if [ ! -f ../../../Binaries/DotNET/GitDependencies/osx-x64/GitDependencies ]; then
 	echo "Cannot find GitDependencies. This script should be placed in Engine/Build/BatchFiles/Mac."
 	exit 1
 fi 
 
 cd ../../../..
 
-dotnet ./Engine/Binaries/DotNET/GitDependencies.dll "$@"
+./Engine/Binaries/DotNET/GitDependencies/osx-x64/GitDependencies "$@"
 
 pushd "$(dirname "$SCRIPT_PATH")" > /dev/null
 sh FixDependencyFiles.sh
diff --git a/Setup.bat b/Setup.bat
--- a/Setup.bat
+++ b/Setup.bat
@@ -10,16 +10,16 @@ set PROMPT_ARGUMENT=--prompt
 :no_prompt_argument
 
 rem Sync the dependencies...
-.\Engine\Binaries\DotNET\GitDependencies.exe %PROMPT_ARGUMENT% %*
+.\Engine\Binaries\DotNET\GitDependencies\win-x64\GitDependencies.exe %PROMPT_ARGUMENT% %*
 if ERRORLEVEL 1 goto error
 
 rem Setup the git hooks...
 if not exist .git\hooks goto no_git_hooks_directory
 echo Registering git hooks...
 echo #!/bin/sh >.git\hooks\post-checkout
-echo Engine/Binaries/DotNET/GitDependencies.exe %* >>.git\hooks\post-checkout
+echo Engine/Binaries/DotNET/GitDependencies/win-x64/GitDependencies.exe %* >>.git\hooks\post-checkout
 echo #!/bin/sh >.git\hooks\post-merge
-echo Engine/Binaries/DotNET/GitDependencies.exe %* >>.git\hooks\post-merge
+echo Engine/Binaries/DotNET/GitDependencies/win-x64/GitDependencies.exe %* >>.git\hooks\post-merge
 :no_git_hooks_directory
 
 rem Install prerequisites...
