Fix Setup.sh looking for a removed file.

Now only fail if both Linux and Mac are missing (there should always be at least one)

diff --git a/Setup.sh b/Setup.sh
--- a/Setup.sh
+++ b/Setup.sh
@@ -5,11 +5,11 @@ set -e
 
 cd "`dirname "$0"`"
 
-if [ ! -f Engine/Binaries/DotNET/GitDependencies.exe ]; then
+if [ ! -f Engine/Binaries/DotNET/GitDependencies/linux-x64/GitDependencies ] && [ ! -f Engine/Binaries/DotNET/GitDependencies/osx-x64/GitDependencies ] ; then
 	echo "GitSetup ERROR: This script does not appear to be located \
-       in the root UE directory and must be run from there."
+	   in the root UE directory and must be run from there."
 	exit 1
-fi 
+fi
 
 function FindPromptOrForceArg()
 {
