Updating .gitignore for HoloLens moving

#rnx

diff --git a/.gitignore b/.gitignore
--- a/.gitignore
+++ b/.gitignore
@@ -196,11 +196,13 @@
 /Engine/Binaries/DotNET/UnrealBuildTool.xml
 /Engine/Binaries/DotNET/AutomationScripts/BuildGraph.Automation.xml
 /Engine/Binaries/DotNET/AutomationScripts/Newtonsoft.Json.xml
-/Engine/Binaries/DotNET/AutomationScripts/HoloLens/System.Runtime.WindowsRuntime.xml
-/Engine/Binaries/DotNET/AutomationScripts/HoloLens/WindowsDevicePortalWrapper.xml
+/Engine/Platforms/Hololens/Binaries/DotNET/AutomationTool/AutomationScripts/Platforms/HoloLens/EpicGames.UHT.xml
+/Engine/Platforms/Hololens/Binaries/DotNET/AutomationTool/AutomationScripts/Platforms/HoloLens/System.Runtime.WindowsRuntime.xml
+/Engine/Platforms/Hololens/Binaries/DotNET/AutomationTool/AutomationScripts/Platforms/HoloLens/WindowsDevicePortalWrapper.xml
 /Engine/Source/Programs/AutomationTool/BuildGraph/BuildGraph.Automation.xml
 
 # Ignore generated files from C# tools
+/Engine/Platforms/**/*.deps.json
 /Engine/Source/Programs/**/*.deps.json
 /Engine/Source/Programs/**/*.runtimeconfig*.json
 /Engine/Source/Programs/**/Properties/launchSettings.json
