Updating .gitignore to exclude auto generated RiderLink

#9206

diff --git a/.gitignore b/.gitignore
--- a/.gitignore
+++ b/.gitignore
@@ -257,6 +257,9 @@ Saved/
 # Ignore CLion directory
 .idea/
 
+# Ignore Rider directories
+/Engine/Plugins/Developer/RiderLink
+
 # Annoying windows files
 [Dd]esktop.ini
 
