PR 9209: Updating engine binaries download size

#rnx
#9209

diff --git a/README.md b/README.md
--- a/README.md
+++ b/README.md
@@ -65,7 +65,7 @@ The steps below take you through cloning your own private fork, then compiling a
 
     On Windows 8, a warning from SmartScreen may appear. Click **More info**, then **Run anyway** to continue.
 
-    A clean download of the engine binaries is currently 3-4 GB, which may take some time to complete. Subsequent runs will be much faster, as they only download new and updated content.
+    A clean download of the engine binaries is currently 20-21 GiB, which may take some time to complete. Subsequent runs will be much faster, as they only download new and updated content.
 
 1.  Run **GenerateProjectFiles.bat** to create project files for the engine. It should take less than a minute to complete.  
 
