Adding the norestart flag to UEPrereqSetup_x64

should address automatic rebooting of the system.

#rnx

diff --git a/Setup.bat b/Setup.bat
--- a/Setup.bat
+++ b/Setup.bat
@@ -24,7 +24,7 @@ echo Engine/Binaries/DotNET/GitDependencies/win-x64/GitDependencies.exe %* >>.gi
 
 rem Install prerequisites...
 echo Installing prerequisites...
-start /wait Engine\Extras\Redist\en-us\UEPrereqSetup_x64.exe /quiet
+start /wait Engine\Extras\Redist\en-us\UEPrereqSetup_x64.exe /quiet /norestart
 
 rem Register the engine installation...
 if not exist .\Engine\Binaries\Win64\UnrealVersionSelector-Win64-Shipping.exe goto :no_unreal_version_selector
