#7586 Update README.md instructions for GitHub Desktop

#rnx

diff --git a/README.md b/README.md
--- a/README.md
+++ b/README.md
@@ -44,9 +44,10 @@ The steps below take you through cloning your own private fork, then compiling a
 
 ### Windows
 
-1.  Install a Git client like **[GitHub for Windows](https://windows.github.com/)**, then **[fork and clone our repository](https://guides.github.com/activities/forking/)**.
+1.  Install **[GitHub Desktop for Windows](https://desktop.github.com/)** then **[fork and clone our repository](https://guides.github.com/activities/forking/)**. 
+    To use Git from the command line, see the [Setting up Git](https://help.github.com/articles/set-up-git/) and [Fork a Repo](https://help.github.com/articles/fork-a-repo/) articles.
 
-    To use Git from the command line instead, see the [Setting up Git](https://help.github.com/articles/set-up-git/) and [Fork a Repo](https://help.github.com/articles/fork-a-repo/) articles.
+    When forking our repository with **[GitHub Desktop for Windows](https://desktop.github.com/)** is complete, GitHub Desktop will ask "**How are you planning to use this fork?**" There are 2 options to choose: `To contribute to the parent project`, and `For my own purposes`. We ask you select `For my own purposes`. This will help ensure any changes you make stay local to your repository, and will help reduce unintended/accidental pull requests to our repository. If you have made modifications in your forked repository that you wish to submit to Unreal Engine, you can make a pull request from your forked repository, to ours.
 
     If you'd prefer not to use Git, you can get the source with the **Download ZIP** button on the right. Note that the zip utility built in to Windows marks the contents of .zip files downloaded from the Internet as unsafe to execute, so right-click the .zip file and select **Properties…** and **Unblock** before decompressing it.
 
@@ -78,9 +79,11 @@ The steps below take you through cloning your own private fork, then compiling a
 
 ### Mac
 
-1.  Install a Git client like **[GitHub for Mac](https://mac.github.com/)**, then **[fork and clone our repository](https://guides.github.com/activities/forking/)**.
+1.  Install **[GitHub Desktop for Mac](https://desktop.github.com/)** then **[fork and clone our repository](https://guides.github.com/activities/forking/)**. 
+
+    To use Git from the Terminal, see the [Setting up Git](https://help.github.com/articles/set-up-git/) and [Fork a Repo](https://help.github.com/articles/fork-a-repo/) articles.
 
-    To use Git from the Terminal instead, see the [Setting up Git](https://help.github.com/articles/set-up-git/) and [Fork a Repo](https://help.github.com/articles/fork-a-repo/) articles.
+    When forking our repository with **[GitHub Desktop for Mac](https://desktop.github.com/)** is complete, GitHub Desktop will ask "**How are you planning to use this fork?**" There are 2 options to choose: `To contribute to the parent project`, and `For my own purposes`. We ask you select `For my own purposes`. This will help ensure any changes you make stay local to your repository, and will help reduce unintended/accidental pull requests to our repository. If you have made modifications in your forked repository that you wish to submit to Unreal Engine, you can make a pull request from your forked repository, to ours.
 
     If you'd rather not use Git, use the **Download ZIP** button on the right to get the source directly.
 
