Allow *.tsx files

#rnx

diff --git a/.gitignore b/.gitignore
--- a/.gitignore
+++ b/.gitignore
@@ -72,6 +72,7 @@
 !*.py
 !*.js
 !*.ts
+!*.tsx
 !*.command
 
 # Other configuration and markup files
