Update EULA, fix broken links, remove EA

diff --git a/README.md b/README.md
--- a/README.md
+++ b/README.md
@@ -7,14 +7,12 @@ With the code in this repository, you can build the Unreal Editor for Windows, M
 
 We have a vast amount of [official documentation](http://docs.unrealengine.com) available for the engine. If you're looking for the answer to something, you may want to start in one of these places:
 
-*   [Programming and Scripting in Unreal Engine](https://docs.unrealengine.com/ProgrammingAndScripting)
-*   [Development Setup](https://docs.unrealengine.com/ProductionPipelines/DevelopmentSetup/)
-*   [Working with the GitHub source code distribution](https://docs.unrealengine.com/ProgrammingAndScripting/ProgrammingWithCPP/DownloadingSourceCode)
+*   [Programming and Scripting in Unreal Engine](https://docs.unrealengine.com/unreal-engine-programming-and-scripting)
+*   [Development Setup](https://docs.unrealengine.com/setting-up-your-development-environment-for-cplusplus-in-unreal-engine/)
+*   [Working with the GitHub source code distribution](https://docs.unrealengine.com/downloading-unreal-engine-source-code)
 *   [Unreal Engine C++ API Reference](https://docs.unrealengine.com/API)
-*   [Community forum about working with engine source and GitHub](https://forums.unrealengine.com/c/development-discussion/c-programming/engine-source-github/33)
 
-If you need more, just ask! Many Epic developers read the [forums](https://forums.unrealengine.com/) and [AnswerHub](https://answers.unrealengine.com/),
-and we're proud to be part of a well-meaning, friendly, and welcoming community of thousands.
+If you need more, just ask! Many Epic developers read the [Discussion](https://forums.unrealengine.com/latest?exclude_tag=question) and [Q&A](https://forums.unrealengine.com/tag/question) forums on the [Epic Games Dev Community](https://dev.epicgames.com/community/) site, and we're proud to be part of a well-meaning, friendly, and welcoming community of thousands.
 
 
 Branches
@@ -28,8 +26,6 @@ We publish source for the engine in several branches:
 
     If you choose to work in this branch, be aware that it is likely to be **ahead** of the branches for the current official release and the next upcoming release. Therefore, content and code that you create to work with the ue5-main branch may not be compatible with public releases until we create a new branch directly from ue5-main for a future official release.
 
-*   For a stable Early Access release of UE5, use the **[ue5-early-access](https://github.com/EpicGames/UnrealEngine/tree/ue5-early-access)** branch. The code in this branch matches the UE5 Early Access binary release that is available in the Epic Games Launcher. It’s a good place to start if you want a solid foundation for trying out the new features in UE5. For more information, see [the UE5 documentation](https://docs.unrealengine.com/5.0/en-US/Welcome/).
-
 *   The **[master branch](https://github.com/EpicGames/UnrealEngine/tree/master)** is the hub of changes to UE4 from all our engine development teams. It’s not subject to as much testing as release branches, so if you’re using UE4 this may not be the best choice for you.
 
 *   Branches whose names contain `dev`, `staging`, and `test` are typically for internal Epic processes, and are rarely useful for end users.
@@ -45,11 +41,14 @@ The steps below take you through cloning your own private fork, then compiling a
 ### Windows
 
 1.  Install **[GitHub Desktop for Windows](https://desktop.github.com/)** then **[fork and clone our repository](https://guides.github.com/activities/forking/)**. 
-    To use Git from the command line, see the [Setting up Git](https://help.github.com/articles/set-up-git/) and [Fork a Repo](https://help.github.com/articles/fork-a-repo/) articles.
 
-    When forking our repository with **[GitHub Desktop for Windows](https://desktop.github.com/)** is complete, GitHub Desktop will ask "**How are you planning to use this fork?**" There are 2 options to choose: `To contribute to the parent project`, and `For my own purposes`. We ask you select `For my own purposes`. This will help ensure any changes you make stay local to your repository, and will help reduce unintended/accidental pull requests to our repository. If you have made modifications in your forked repository that you wish to submit to Unreal Engine, you can make a pull request from your forked repository, to ours.
+    When GitHub Desktop is finished creating the fork, it asks, **How are you planning to use this fork?**. Select **For my own purposes**. This will help ensure that any changes you make stay local to your repository and avoid creating unintended pull requests. You'll still be able to make pull requests if you want to submit modifications that you make in your fork back to our repository.
+
+    Other options:
 
-    If you'd prefer not to use Git, you can get the source with the **Download ZIP** button on the right. Note that the zip utility built in to Windows marks the contents of .zip files downloaded from the Internet as unsafe to execute, so right-click the .zip file and select **Properties…** and **Unblock** before decompressing it.
+    -   To use Git from the command line, see the [Setting up Git](https://help.github.com/articles/set-up-git/) and [Fork a Repo](https://help.github.com/articles/fork-a-repo/) articles.
+
+    -   If you'd prefer not to use Git, you can get the source with the **Download ZIP** button on the right. Note that the zip utility built in to Windows marks the contents of .zip files downloaded from the Internet as unsafe to execute, so right-click the .zip file and select **Properties…** and **Unblock** before decompressing it.
 
 1.  Install **Visual Studio 2019**.
 
@@ -81,11 +80,13 @@ The steps below take you through cloning your own private fork, then compiling a
 
 1.  Install **[GitHub Desktop for Mac](https://desktop.github.com/)** then **[fork and clone our repository](https://guides.github.com/activities/forking/)**. 
 
-    To use Git from the Terminal, see the [Setting up Git](https://help.github.com/articles/set-up-git/) and [Fork a Repo](https://help.github.com/articles/fork-a-repo/) articles.
+    When GitHub Desktop is finished creating the fork, it asks, **How are you planning to use this fork?**. Select **For my own purposes**. This will help ensure that any changes you make stay local to your repository and avoid creating unintended pull requests. You'll still be able to make pull requests if you want to submit modifications that you make in your fork back to our repository.
+
+    Other options:
 
-    When forking our repository with **[GitHub Desktop for Mac](https://desktop.github.com/)** is complete, GitHub Desktop will ask "**How are you planning to use this fork?**" There are 2 options to choose: `To contribute to the parent project`, and `For my own purposes`. We ask you select `For my own purposes`. This will help ensure any changes you make stay local to your repository, and will help reduce unintended/accidental pull requests to our repository. If you have made modifications in your forked repository that you wish to submit to Unreal Engine, you can make a pull request from your forked repository, to ours.
+    -   To use Git from the Terminal, see the [Setting up Git](https://help.github.com/articles/set-up-git/) and [Fork a Repo](https://help.github.com/articles/fork-a-repo/) articles.
 
-    If you'd rather not use Git, use the **Download ZIP** button on the right to get the source directly.
+    -   If you'd prefer not to use Git, use the **Download ZIP** button on the right to get the source directly.
 
 1.  Install the latest version of [Xcode](https://itunes.apple.com/us/app/xcode/id497799835).
 
@@ -104,45 +105,47 @@ The steps below take you through cloning your own private fork, then compiling a
 
 1.  Install a **[visual Git client](https://git-scm.com/download/gui/linux)**, then **[fork and clone our repository](https://guides.github.com/activities/forking/)**.
 
-    To use Git from the command line instead, see the [Setting up Git](https://help.github.com/articles/set-up-git/) and [Fork a Repo](https://help.github.com/articles/fork-a-repo/) articles.
+    Other options:
+    
+    -   To use Git from the command line instead, see the [Setting up Git](https://help.github.com/articles/set-up-git/) and [Fork a Repo](https://help.github.com/articles/fork-a-repo/) articles.
 
-    If you'd prefer not to use Git, use the **Download ZIP** button on the right to get the source as a zip file.
+    -   If you'd prefer not to use Git, use the **Download ZIP** button on the right to get the source as a zip file.
 
 1.  Open your source folder and run **Setup.sh** to download binary content for the engine.
 
 1.  Both cross-compiling and native builds are supported.
 
-    -   **Cross-compiling** is handy for Windows developers who want to package a game for Linux with minimal hassle. It requires a cross-compiler toolchain to be installed. See the [Linux cross-compiling page in the documentation](https://docs.unrealengine.com/SharingAndReleasing/Linux/GettingStarted/).
+    -   **Cross-compiling** is handy for Windows developers who want to package a game for Linux with minimal hassle. It requires a cross-compiler toolchain to be installed. See the [Linux cross-compiling page in the documentation](https://docs.unrealengine.com/linux-development-requirements-for-unreal-engine/).
 
     -   **Native compilation** is discussed in [a separate README](Engine/Build/BatchFiles/Linux/README.md) and [community wiki page](https://unrealcommunity.wiki/building-on-linux-qr8t0si2).
 
 
 ### Additional target platforms
 
-*   **Android** support will be downloaded by the setup script if you have the Android NDK installed. See the [Android Quick Start guide](https://docs.unrealengine.com/en-US/SharingAndReleasing/Mobile/Android/GettingStarted/).
+*   **Android** support will be downloaded by the setup script if you have the Android NDK installed. See the [Android Quick Start guide](https://docs.unrealengine.com/setting-up-unreal-engine-projects-for-android-development/).
 
-*   **iOS** development requires a Mac. Instructions are in the [iOS Quick Start guide](https://docs.unrealengine.com/en-US/SharingAndReleasing/Mobile/iOS/QuickStart/).
+*   **iOS** development requires a Mac. Instructions are in the [iOS Quick Start guide](https://docs.unrealengine.com/setting-up-an-unreal-engine-project-for-ios/).
 
 *   Development for consoles and other platforms with restricted access, like **Sony PlayStation**, **Microsoft Xbox**, and **Nintendo Switch**, is only possible if you have a registered developer account with those third-party vendors.
 
-    Depending on the platform, additional documentation or guidance may be available in the [Unreal Developer Network support site](https://udn.unrealengine.com/s/), or as a downloadable archive in the section of the [Unreal Engine Forums](https://forums.unrealengine.com/) that is dedicated to your platform.
+    Depending on the platform, additional documentation or guidance may be available in the [Unreal Developer Network support site](https://udn.unrealengine.com/s/), or as a downloadable archive in the section of the [Unreal Engine Forums](https://forums.unrealengine.com/c/development-discussion/platforms/130) that is dedicated to your platform.
 
-    If you don’t have access to these resources, first register a developer account with the third party vendor. Then contact your Epic Games account manager if you have one, or fill out and submit the [Console Development Request form](https://epicgames.secure.force.com/Forms/FormConsoleAccessRequest?) for Unreal Engine if you don’t. Epic will contact you with a formal agreement to digitally sign. Once this is approved, you will receive instructions on how to access source code, binaries, and additional instructions for your platform.
+    If you don’t have access to these resources, first register a developer account with the third party vendor. Then contact your Epic Games account manager if you have one, or fill out and submit the [Console Development Request form](https://forms.unrealengine.com/s/form-console-access-request) for Unreal Engine if you don’t. Epic will contact you with a formal agreement to digitally sign. Once this is approved, you will receive instructions on how to access source code, binaries, and additional instructions for your platform.
 
 
 Licensing
 ---------
 
-Your access to and use of Unreal Engine on GitHub is governed by an End User License Agreement (EULA). You can choose either a Publishing EULA or a Creators EULA, depending on how the terms fit your intended usage of the engine. For a comparison and links to FAQs, see the [Unreal Engine download page](https://www.unrealengine.com/en-US/download). If you don't agree to the terms in your chosen EULA, as amended from time to time, you are not permitted to access or use Unreal Engine.
+Your access to and use of Unreal Engine on GitHub is governed by an End User License Agreement (EULA). For the latest terms and conditions, see the license and FAQ on the [Unreal Engine download page](https://www.unrealengine.com/download). If you don't agree to the terms in your chosen EULA, as amended from time to time, you are not permitted to access or use Unreal Engine.
 
 Contributions
 -------------
 
 We welcome contributions to Unreal Engine development through [pull requests](https://github.com/EpicGames/UnrealEngine/pulls/) on GitHub.
 
-We prefer to take pull requests in our active development branches, particularly for new features. For UE5, use the **ue5-main** branch; for UE4, use the **master** branch. Please make sure that all new code adheres to the [Epic coding standards](https://docs.unrealengine.com/ProductionPipelines/DevelopmentSetup/CodingStandard).
+We prefer to take pull requests in our active development branches, particularly for new features. For UE5, use the **ue5-main** branch; for UE4, use the **master** branch. Please make sure that all new code adheres to the [Epic coding standards](https://docs.unrealengine.com/epic-cplusplus-coding-standard-for-unreal-engine/).
 
-For more information on the process and expectations, see [the documentation](https://docs.unrealengine.com/ProgrammingAndScripting/ProgrammingWithCPP/DownloadingSourceCode/ContributingToUnreal/).
+For more information on the process and expectations, see [the documentation](https://docs.unrealengine.com/contributing-to-the-unreal-engine/).
 
 All contributions are governed by the terms of your EULA.
 
@@ -150,6 +153,6 @@ All contributions are governed by the terms of your EULA.
 Additional Notes
 ----------------
 
-The first time you start the editor from a fresh source build, you may experience long load times. The engine is optimizing content for your platform and storing it in the _[derived data cache](https://docs.unrealengine.com/ProductionPipelines/DerivedDataCache/)_. This should only happen once.
+The first time you start the editor from a fresh source build, you may experience long load times. The engine is optimizing content for your platform and storing it in the _[derived data cache](https://docs.unrealengine.com/derived-data-cache/)_. This should only happen once.
 
 Your private forks of the Unreal Engine code are associated with your GitHub account permissions. If you unsubscribe or switch GitHub user names, you'll need to create a new fork and upload your changes from the fresh copy.
