Adding some dockerfiles to the repo

#rnx

diff --git a/.gitignore b/.gitignore
--- a/.gitignore
+++ b/.gitignore
@@ -186,6 +186,9 @@
 !/Engine/Binaries/DotNET/GitDependencies.dll
 !/Engine/Binaries/DotNET/GitDependencies.runtimeconfig.json
 
+# Add Horde dockerfiles
+!/Engine/Source/Programs/Horde/**/Dockerfile
+
 # Ignore NET binaries
 /Engine/Binaries/ThirdParty/DotNet/
 /Engine/Binaries/DotNET/UnrealBuildTool/
