Fixing Setup.bat not exiting if GitDependencies throws an exception

#rnx

diff --git a/Setup.bat b/Setup.bat
--- a/Setup.bat
+++ b/Setup.bat
@@ -11,7 +11,7 @@ set PROMPT_ARGUMENT=--prompt
 
 rem Sync the dependencies...
 .\Engine\Binaries\DotNET\GitDependencies\win-x64\GitDependencies.exe %PROMPT_ARGUMENT% %*
-if ERRORLEVEL 1 goto error
+if %ERRORLEVEL% NEQ 0 goto error
 
 rem Setup the git hooks...
 if not exist .git\hooks goto no_git_hooks_directory
