#9612 Don't ignore poweshell scripts

#rnx

diff --git a/.gitignore b/.gitignore
--- a/.gitignore
+++ b/.gitignore
@@ -70,6 +70,7 @@
 !*.sh
 !*.pl
 !*.py
+!*.ps1
 !*.js
 !*.ts
 !*.tsx
