Updating gitignore to ignore some json/xml files

#rnx

diff --git a/.gitignore b/.gitignore
--- a/.gitignore
+++ b/.gitignore
@@ -201,9 +201,7 @@
 /Engine/Binaries/DotNET/UnrealBuildTool.xml
 /Engine/Binaries/DotNET/AutomationScripts/BuildGraph.Automation.xml
 /Engine/Binaries/DotNET/AutomationScripts/Newtonsoft.Json.xml
-/Engine/Platforms/Hololens/Binaries/DotNET/AutomationTool/AutomationScripts/Platforms/HoloLens/EpicGames.UHT.xml
-/Engine/Platforms/Hololens/Binaries/DotNET/AutomationTool/AutomationScripts/Platforms/HoloLens/System.Runtime.WindowsRuntime.xml
-/Engine/Platforms/Hololens/Binaries/DotNET/AutomationTool/AutomationScripts/Platforms/HoloLens/WindowsDevicePortalWrapper.xml
+/Engine/Platforms/Hololens/Binaries/DotNET/AutomationTool/AutomationScripts/Platforms/HoloLens/*.xml
 /Engine/Source/Programs/AutomationTool/BuildGraph/BuildGraph.Automation.xml
 
 # Ignore generated files from C# tools
@@ -276,3 +274,4 @@ Engine/Extras/ThirdPartyNotUE/SDKs/HostLinux/**
 
 # Ignore test json data files in NeuralNetworkInterface
 /Engine/Plugins/Experimental/NNI/Source/ThirdParty/Deps/json/**
+/Engine/Plugins/Experimental/NNX/Source/ThirdParty/Deps/json/**
