Updating gitignore entry for renamed NN plugin

#rnx

diff --git a/.gitignore b/.gitignore
--- a/.gitignore
+++ b/.gitignore
@@ -274,4 +274,4 @@ Engine/Extras/ThirdPartyNotUE/SDKs/HostLinux/**
 
 # Ignore test json data files in NeuralNetworkInterface
 /Engine/Plugins/Experimental/NNI/Source/ThirdParty/Deps/json/**
-/Engine/Plugins/Experimental/NNX/Source/ThirdParty/Deps/json/**
+/Engine/Plugins/Experimental/NNE/Source/ThirdParty/Deps/json/**
