Updating gitignore to allow yaml & dockerfiles into the repo

#rnx

diff --git a/.gitignore b/.gitignore
--- a/.gitignore
+++ b/.gitignore
@@ -91,8 +91,11 @@
 !*.config
 !*.version
 !.git*
+!.dockerignore
 !.editorconfig
 !.plist
+!.yaml
+!.yml
 
 # Projects and makefiles
 !*.cmake
@@ -108,6 +111,7 @@
 !*.xcconfig
 !*.vsprops
 !*.snippet
+!Dockerfile
 !Makefile
 !Makefile.*
 !Settings.settings
