Update git README.md to reference Visual Studio setup documentation

#rnx

diff --git a/README.md b/README.md
--- a/README.md
+++ b/README.md
@@ -5,7 +5,7 @@ Welcome to the Unreal Engine source code!
 
 With the code in this repository, you can build the Unreal Editor for Windows, Mac, and Linux; compile Unreal Engine games for a variety of target platforms, including desktop, consoles, mobile, and embedded devices; and build tools like Unreal Lightmass and Unreal Frontend. Modify the code in any way you can imagine, and share your changes with others!
 
-We have a vast amount of [official documentation](http://docs.unrealengine.com) available for the engine. If you're looking for the answer to something, you may want to start in one of these places:
+We have a vast amount of [official documentation](https://docs.unrealengine.com) available for the engine. If you're looking for the answer to something, you may want to start in one of these places:
 
 *   [Programming and Scripting in Unreal Engine](https://docs.unrealengine.com/unreal-engine-programming-and-scripting)
 *   [Development Setup](https://docs.unrealengine.com/setting-up-your-development-environment-for-cplusplus-in-unreal-engine/)
@@ -50,16 +50,11 @@ The steps below take you through cloning your own private fork, then compiling a
 
     -   If you'd prefer not to use Git, you can get the source with the **Download ZIP** button on the right. Note that the zip utility built in to Windows marks the contents of .zip files downloaded from the Internet as unsafe to execute, so right-click the .zip file and select **Properties…** and **Unblock** before decompressing it.
 
-1.  Install **Visual Studio 2019**.
+1.  Install **Visual Studio 2022** or **Visual Studion 2019**.
 
-    All desktop editions of Visual Studio 2019 can build UE5, including [Visual Studio Community 2019](http://www.visualstudio.com/products/visual-studio-community-vs), which is free for small teams and individual developers.
+    All desktop editions of Visual Studio 2022 and 2019 can build UE5, including [Visual Studio Community](https://www.visualstudio.com/products/visual-studio-community-vs), which is free for small teams and individual developers.
 
-    To install the correct components for UE5 development, make sure the **Game Development with C++** workload is checked. Under the **Installation Details** section on the right, also choose the following components:
-
-    -   **C++ profiling tools**
-    -   **C++ AddressSanitizer** (optional)
-    -   **Windows 10 SDK** (10.0.18362 or newer)
-    -   **Unreal Engine Installer**
+    To install the correct components for UE5 development, please see [Setting Up Visual Studio](https://docs.unrealengine.com/en-US/setting-up-visual-studio-development-environment-for-cplusplus-projects-in-unreal-engine/).
 
 1.  Open your source folder in Windows Explorer and run **Setup.bat**. This will download binary content for the engine, install prerequisites, and set up Unreal file associations.
 
