Git: Update ignore file for new exclusion

#rnx

diff --git a/.gitignore b/.gitignore
--- a/.gitignore
+++ b/.gitignore
@@ -279,3 +279,4 @@ Engine/Extras/ThirdPartyNotUE/SDKs/HostLinux/**
 # Ignore test json data files in NeuralNetworkInterface
 /Engine/Plugins/Experimental/NNI/Source/ThirdParty/Deps/json/**
 /Engine/Plugins/Experimental/NNE/Source/ThirdParty/Deps/json/**
+/Engine/Source/ThirdParty/Vulkan/share/vulkan/registry/**
