PR #10139: Use git to find the GIT_DIR (Contributed by [REDACTED_PERSON])

diff --git a/Engine/Build/BatchFiles/Linux/SetupToolchain.sh b/Engine/Build/BatchFiles/Linux/SetupToolchain.sh
--- a/Engine/Build/BatchFiles/Linux/SetupToolchain.sh
+++ b/Engine/Build/BatchFiles/Linux/SetupToolchain.sh
@@ -4,12 +4,19 @@
 # put ourselves into Engine directory (two up from location of this script)
 pushd "`dirname "$0"`/../../.." > /dev/null
 
+if [ -z $GIT_DIR ]; then
+	GIT_DIR=$(git rev-parse --git-common-dir);
+	if [ $? -ne 0 ]; then
+		GIT_DIR=../.git
+	fi
+fi
+
 TOOLCHAIN_VERSION=v21_clang-15.0.1-centos7
 TOOLCHAIN_ARCHIVE=$TOOLCHAIN_VERSION.tar.gz
 
 TOOLCHAIN_URL=http://cdn.unrealengine.com/Toolchain_Linux/native-linux-$TOOLCHAIN_ARCHIVE
 TOOLCHAIN_ROOT=Extras/ThirdPartyNotUE/SDKs/HostLinux/Linux_x64/
-TOOLCHAIN_CACHE=../.git/ue4-sdks/
+TOOLCHAIN_CACHE="$GIT_DIR/ue4-sdks/"
 
 if [ -z $UE_SDK_CACHE_SIZE ]; then
 TOOLCHAIN_CACHE_SIZE=2
diff --git a/Setup.sh b/Setup.sh
--- a/Setup.sh
+++ b/Setup.sh
@@ -1,6 +1,13 @@
 #!/bin/bash
 # Copyright Epic Games, Inc. All Rights Reserved.
 
+if [ -z $GIT_DIR ]; then
+	GIT_DIR=$(git rev-parse --git-common-dir);
+	if [ $? -ne 0 ]; then
+		GIT_DIR=.git
+	fi
+fi
+
 set -e
 
 cd "`dirname "$0"`"
@@ -40,27 +47,27 @@ GitDepsArgs+=("$@")
 
 if [ "$(uname)" = "Darwin" ]; then
 	# Setup the git hooks
-	if [ -d .git/hooks ]; then
+	if [ -d "$GIT_DIR/hooks" ]; then
 		echo "Registering git hooks... (this will override existing ones!)"
-		rm -f .git/hooks/post-checkout
-		rm -f .git/hooks/post-merge
-		ln -s ../../Engine/Build/BatchFiles/Mac/GitDependenciesHook.sh .git/hooks/post-checkout
-		ln -s ../../Engine/Build/BatchFiles/Mac/GitDependenciesHook.sh .git/hooks/post-merge
+		rm -f "$GIT_DIR/hooks/post-checkout"
+		rm -f "$GIT_DIR/hooks/post-merge"
+		ln -s ../../Engine/Build/BatchFiles/Mac/GitDependenciesHook.sh "$GIT_DIR/hooks/post-checkout"
+		ln -s ../../Engine/Build/BatchFiles/Mac/GitDependenciesHook.sh "$GIT_DIR/hooks/post-merge"
 	fi
 
 	# Get the dependencies for the first time
 	Engine/Build/BatchFiles/Mac/GitDependencies.sh "${GitDepsArgs[@]}"
 else
 	# Setup the git hooks
-	if [ -d .git/hooks ]; then
+	if [ -d "$GIT_DIR/hooks" ]; then
 		echo "Registering git hooks... (this will override existing ones!)"
-		echo \#!/bin/sh >.git/hooks/post-checkout
-		echo Engine/Build/BatchFiles/Linux/GitDependencies.sh >>.git/hooks/post-checkout
-		chmod +x .git/hooks/post-checkout
+		echo \#!/bin/sh > "$GIT_DIR/hooks/post-checkout"
+		echo Engine/Build/BatchFiles/Linux/GitDependencies.sh >> "$GIT_DIR/hooks/post-checkout"
+		chmod +x "$GIT_DIR/hooks/post-checkout"
 
-		echo \#!/bin/sh >.git/hooks/post-merge
-		echo Engine/Build/BatchFiles/Linux/GitDependencies.sh >>.git/hooks/post-merge
-		chmod +x .git/hooks/post-merge
+		echo \#!/bin/sh > "$GIT_DIR/hooks/post-merge"
+		echo Engine/Build/BatchFiles/Linux/GitDependencies.sh >> "$GIT_DIR/hooks/post-merge"
+		chmod +x "$GIT_DIR/hooks/post-merge"
 	fi
 
 	# Get the dependencies for the first time
