Git: Add Engine/Programs so program project inis are picked up properly.

diff --git a/Default.uprojectdirs b/Default.uprojectdirs
--- a/Default.uprojectdirs
+++ b/Default.uprojectdirs
@@ -7,3 +7,4 @@
 ;
 ./
 Engine/Source/
+Engine/Programs/
\ No newline at end of file
