Cleaning up NNI and NNE specific gitignore paths (the files have been removed completely)

#rnx

diff --git a/.gitignore b/.gitignore
--- a/.gitignore
+++ b/.gitignore
@@ -276,7 +276,5 @@ Engine/Extras/ThirdPartyNotUE/SDKs/HostLinux/**
 # Ignore USD resources copied into Binaries directory at build time
 /Engine/Binaries/ThirdParty/USD/**
 
-# Ignore test json data files in NeuralNetworkInterface
-/Engine/Plugins/Experimental/NNI/Source/ThirdParty/Deps/json/**
-/Engine/Plugins/Experimental/NNE/Source/ThirdParty/Deps/json/**
+# Ignore vulkan registry
 /Engine/Source/ThirdParty/Vulkan/share/vulkan/registry/**
