PR #10718: Updated README.md in release branch

diff --git a/README.md b/README.md
--- a/README.md
+++ b/README.md
@@ -50,7 +50,7 @@ The steps below take you through cloning your own private fork, then compiling a
 
     -   If you'd prefer not to use Git, you can get the source with the **Download ZIP** button on the right. Note that the zip utility built in to Windows marks the contents of .zip files downloaded from the Internet as unsafe to execute, so right-click the .zip file and select **Properties…** and **Unblock** before decompressing it.
 
-1.  Install **Visual Studio 2022** or **Visual Studion 2019**.
+1.  Install **Visual Studio 2022** or **Visual Studio 2019**.
 
     All desktop editions of Visual Studio 2022 and 2019 can build UE5, including [Visual Studio Community](https://www.visualstudio.com/products/visual-studio-community-vs), which is free for small teams and individual developers.
 
