Move verse editor configuration to the UE editor config.

diff --git a/.editorconfig b/.editorconfig
--- a/.editorconfig
+++ b/.editorconfig
@@ -31,3 +31,10 @@ indent_size = 4
 [*.yaml]
 indent_style = space
 indent_size = 2
+
+# Settings for Verse source files
+[*.{verse,versetest}]
+indent_style = space
+indent_size = 4
+tab_width = 4
+max_line_length = off
