fix for git untracked files

diff --git a/.gitignore b/.gitignore
--- a/.gitignore
+++ b/.gitignore
@@ -231,6 +231,16 @@ Engine/Binaries/Android/packageInfo.txt
 # Ignore version files in the Engine/Binaries directory created by UBT
 /Engine/Binaries/**/*.version
 
+# Ignore built Mac apps
+/Engine/Binaries/Mac/UnrealEditor*.app/**
+/Engine/Binaries/Mac/UnrealGame*.app/**
+
+# Ignore generated Xcode settings file
+/Engine/Build/Xcode/ContentOnlySettings.xcconfig
+
+# Ignore .build folder
+/Engine/Binaries/*.build/**
+
 # Ignore exp files in Engine and Plugin Binaries directories as they aren't C/C++ source files
 /Engine/**/Binaries/**/*.exp
 
