Add some rider code-style settings to the editor config

diff --git a/.editorconfig b/.editorconfig
--- a/.editorconfig
+++ b/.editorconfig
@@ -10,7 +10,13 @@ indent_style = tab
 indent_size = 4
 tab_width = 4
 trim_trailing_whitespace = false
-max_line_length = 120
+max_line_length = 150
+
+# ReSharper properties
+resharper_cpp_keep_user_linebreaks = false
+resharper_cpp_wrap_lines = false
+resharper_cpp_wrap_before_colon = true
+resharper_cpp_line_break_before_comma_in_member_initializer_lists = true
 
 # MSBuild
 [*.{csproj,proj,projitems,shproj,fsproj,target,props}]
