Fix untracked file CLI11/meson.build

diff --git a/.gitignore b/.gitignore
--- a/.gitignore
+++ b/.gitignore
@@ -291,3 +291,6 @@ Engine/Extras/ThirdPartyNotUE/SDKs/HostLinux/**
 
 # Ignore vulkan registry
 /Engine/Source/ThirdParty/Vulkan/share/vulkan/registry/**
+
+# Do NOT ignore meson.build inside CLI11, these are included in official CLI11.git repo
+!/Engine/Source/ThirdParty/CLI11/git-main/**/meson.build
