Updating VS section of the GitHub Readme

#rnx

diff --git a/README.md b/README.md
--- a/README.md
+++ b/README.md
@@ -50,9 +50,9 @@ The steps below take you through cloning your own private fork, then compiling a
 
     -   If you'd prefer not to use Git, you can get the source with the **Download ZIP** button on the right. Note that the zip utility built in to Windows marks the contents of .zip files downloaded from the Internet as unsafe to execute, so right-click the .zip file and select **Properties…** and **Unblock** before decompressing it.
 
-1.  Install **Visual Studio 2022** or **Visual Studio 2019**.
+1.  Install **Visual Studio 2022**.
 
-    All desktop editions of Visual Studio 2022 and 2019 can build UE5, including [Visual Studio Community](https://www.visualstudio.com/products/visual-studio-community-vs), which is free for small teams and individual developers.
+    All desktop editions of Visual Studio 2022, **17.4** or later, **17.8** recommended, can build UE5, including [Visual Studio Community](https://www.visualstudio.com/products/visual-studio-community-vs), which is free for small teams and individual developers.
 
     To install the correct components for UE5 development, please see [Setting Up Visual Studio](https://docs.unrealengine.com/en-US/setting-up-visual-studio-development-environment-for-cplusplus-projects-in-unreal-engine/).
 
@@ -60,7 +60,7 @@ The steps below take you through cloning your own private fork, then compiling a
 
     On Windows 8, a warning from SmartScreen may appear. Click **More info**, then **Run anyway** to continue.
 
-    A clean download of the engine binaries is currently 20-21 GiB, which may take some time to complete. Subsequent runs will be much faster, as they only download new and updated content.
+    A clean download of the engine binaries is currently 18+ GiB, which may take some time to complete. Subsequent runs will be much faster, as they only download new and updated content.
 
 1.  Run **GenerateProjectFiles.bat** to create project files for the engine. It should take less than a minute to complete.  
 
