Updating gitignore for a new path

#rnx

diff --git a/.gitignore b/.gitignore
--- a/.gitignore
+++ b/.gitignore
@@ -165,6 +165,9 @@
 # Exception for VP EpicStageApp dart files
 !/Engine/Extras/VirtualProduction/EpicStageApp/**/*.dart
 
+# Exception for VS Workspace patches
+!/Engine/Extras/VisualStudioWorkspace/
+
 # Ignore the automatically created mono symlink
 /Engine/Binaries/ThirdParty/Mono/Mac/bin/mono
 
