Updating GitHub Readme

#rnx

diff --git a/README.md b/README.md
--- a/README.md
+++ b/README.md
@@ -3,6 +3,8 @@ Unreal Engine
 
 Welcome to the Unreal Engine source code!
 
+UE 5.4 introduced a new EULA that will govern this and subsequent releases. Depending on the work you create with Unreal Engine, you might be subject to new or different fees. Please review the [EULA here](https://www.unrealengine.com/eula) and make sure you understand how these changes apply to you. By downloading or using the UE 5.4 or later, or by otherwise indicating your acceptance of the EULA, you are agreeing to be bound by its terms. If you do not or cannot agree to the terms of the EULA, please do not download or use Unreal Engine 5.4 or later. 
+
 With the code in this repository, you can build the Unreal Editor for Windows, Mac, and Linux; compile Unreal Engine games for a variety of target platforms, including desktop, consoles, mobile, and embedded devices; and build tools like Unreal Lightmass and Unreal Frontend. Modify the code in any way you can imagine, and share your changes with others!
 
 We have a vast amount of [official documentation](https://docs.unrealengine.com) available for the engine. If you're looking for the answer to something, you may want to start in one of these places:
