Add copyright header to .editorconfig

diff --git a/.editorconfig b/.editorconfig
--- a/.editorconfig
+++ b/.editorconfig
@@ -44,3 +44,6 @@ indent_style = space
 indent_size = 4
 tab_width = 4
 max_line_length = off
+
+# Copyright
+file_header_template =// Copyright Epic Games, Inc. All Rights Reserved.
