Added cpp_indent_lambda_braces_when_parameter = true to .editorconfig to better match our intended syntax.

diff --git a/.editorconfig b/.editorconfig
--- a/.editorconfig
+++ b/.editorconfig
@@ -11,6 +11,7 @@ indent_size = 4
 tab_width = 4
 trim_trailing_whitespace = false
 max_line_length = 150
+cpp_indent_lambda_braces_when_parameter = true
 
 # ReSharper properties
 resharper_cpp_keep_user_linebreaks = false
