GitHub: Ignore some large data files not needed int the repo directly

#rnx

diff --git a/.gitignore b/.gitignore
--- a/.gitignore
+++ b/.gitignore
@@ -294,3 +294,7 @@ Engine/Extras/ThirdPartyNotUE/SDKs/HostLinux/**
 
 # Ignore vulkan registry
 /Engine/Source/ThirdParty/Vulkan/share/vulkan/registry/**
+
+# Ignore large FB data files
+/Engine/Source/ThirdParty/Facebook/IOS/**/*.yml
+/Engine/Source/ThirdParty/Facebook/IOS/**/*.json
