[Rider] Editor Config change to keep user linebreaks, and wrap lines.
 This formats more consistently with how Cvars and other very long constructors/initializers are formatted.  Sometimes, such as in the cvar case, users want to force linebreaks to get more readability, so we shouldn't wipe those out.  As well, the files I've seen, all wrap lines, at around 150, so we keep that from the last change.

#tests None required, it's just an editor config.

diff --git a/.editorconfig b/.editorconfig
--- a/.editorconfig
+++ b/.editorconfig
@@ -14,8 +14,8 @@ max_line_length = 150
 cpp_indent_lambda_braces_when_parameter = true
 
 # ReSharper properties
-resharper_cpp_keep_user_linebreaks = false
-resharper_cpp_wrap_lines = false
+resharper_cpp_keep_user_linebreaks = true
+resharper_cpp_wrap_lines = true
 resharper_cpp_wrap_before_colon = true
 resharper_cpp_line_break_before_comma_in_member_initializer_lists = true
 
