PR #12342: Adding syntax highlighting for shaders in GitHub

#rnx

diff --git a/.gitattributes b/.gitattributes
--- a/.gitattributes
+++ b/.gitattributes
@@ -10,3 +10,7 @@
 
 *.bat eol=crlf
 *.sh eol=lf
+
+# Enable github syntax highlighting for shader files.
+*.usf linguist-language=HLSL
+*.ush linguist-language=HLSL
