Adding *.axaml to the allowlist for git repo files

#rnx

diff --git a/.gitignore b/.gitignore
--- a/.gitignore
+++ b/.gitignore
@@ -82,6 +82,7 @@
 !*.tps
 !*.xml
 !*.xaml
+!*.axaml
 !*.uproject
 !*.uplugin
 !*.html
