Fixing gitignore after UAT/UBT compile location changes

#rnx

diff --git a/.gitignore b/.gitignore
--- a/.gitignore
+++ b/.gitignore
@@ -207,6 +207,8 @@
 /Engine/Binaries/DotNET/AutomationTool/
 /Engine/Source/Programs/DotNETCommon/BuildUtilities/bin/
 /Engine/Source/Programs/Shared/*/bin/
+/Engine/Source/Programs/AutomationTool/**/bin/**/*.xml
+/Engine/Source/Programs/UnrealBuildTool/**/bin/**/*.xml
 
 # Ignore documentation generated for C# tools
 /Engine/Binaries/DotNET/UnrealBuildTool.xml
