Adding Security.MD file on behalf of InfoSec

#rnx

diff --git a/SECURITY.md b/SECURITY.md
new file mode 100644
--- /dev/null
+++ b/SECURITY.md
@@ -0,0 +1,9 @@
+Our Security Approach
+
+The Information Security Team at Epic Games is continually working to protect our players, data, products, and services. Our team understands that there is an incredible community of talented security researchers who want to help further strengthen our security.
+
+If you think you've found a vulnerability on one of our sites or within our products, please report it to us by either emailing [REDACTED_EMAIL], or directly submitting it to our HackerOne bug bounty program (https://hackerone.com/epicgames). Once submitted our team will work with you to validate, resolve, and award your findings as appropriate.
+
+When submitting your vulnerability report, include as much information about your finding as possible - such as domain names, impacted services, and any other relevant information. Please be sure to thoroughly review our program policy prior to submission.
+
+When you submit your vulnerability please be sure to give as much detail as possible. We look forward to hearing from you!
\ No newline at end of file
