Remove GitHub reference to Windows 8 since it is unsupported

Keep the remainder of the message since SmartScreen still applies to Windows 10 and 11.

#rnx

diff --git a/README.md b/README.md
--- a/README.md
+++ b/README.md
@@ -60,7 +60,7 @@ The steps below take you through cloning your own private fork, then compiling a
 
 1.  Open your source folder in Windows Explorer and run **Setup.bat**. This will download binary content for the engine, install prerequisites, and set up Unreal file associations.
 
-    On Windows 8, a warning from SmartScreen may appear. Click **More info**, then **Run anyway** to continue.
+    On Windows, a warning from SmartScreen may appear. Click **More info**, then **Run anyway** to continue.
 
     A clean download of the engine binaries is currently 18+ GiB, which may take some time to complete. Subsequent runs will be much faster, as they only download new and updated content.
 
