Adding PR Template for GitHub

#rnx

diff --git a/PULL_REQUEST_TEMPLATE.md b/PULL_REQUEST_TEMPLATE.md
new file mode 100644
--- /dev/null
+++ b/PULL_REQUEST_TEMPLATE.md
@@ -0,0 +1,9 @@
+## Description
+<!--- Provide a general summary of your change in the Pull Request title above -->
+<!--- Please fill out this form, failure to do so may result in the Pull Request being automatically closed -->
+<!--- Describe your change in detail here. -->
+
+## Checklist:
+<!--- Go over all the following points, and put an `X` in all the boxes that apply like this: [X] -->
+- [ ] My code adheres to the **[Epic coding standards](https://docs.unrealengine.com/latest/INT/Programming/Development/CodingStandard/)**
+- [ ] This change does not include source code created by Generative AI
