[Rider] Change EditorConfig to not align overflowing function calls, and to indent lambdas.

diff --git a/.editorconfig b/.editorconfig
--- a/.editorconfig
+++ b/.editorconfig
@@ -14,10 +14,14 @@ max_line_length = 150
 cpp_indent_lambda_braces_when_parameter = true
 
 # ReSharper properties
+resharper_cpp_align_multiline_argument = false
+resharper_cpp_anonymous_method_declaration_braces = next_line_shifted_2
 resharper_cpp_keep_user_linebreaks = true
 resharper_cpp_wrap_lines = true
 resharper_cpp_wrap_before_colon = true
 resharper_cpp_line_break_before_comma_in_member_initializer_lists = true
+resharper_simple_block_style = line_break
+resharper_line_break_after_deref_in_trailing_return_types = on_single_line
 
 # MSBuild
 [*.{csproj,proj,projitems,shproj,fsproj,target,props}]
