[Rider Editor Config] Make it so ternaries are only chopped (each term on a separate line) when the whole expression exceeds the max line length.

#Synced-CL 52052884

diff --git a/.editorconfig b/.editorconfig
--- a/.editorconfig
+++ b/.editorconfig
@@ -21,7 +21,7 @@ resharper_cpp_keep_user_linebreaks = true
 resharper_cpp_wrap_lines = true
 resharper_cpp_wrap_before_colon = true
 resharper_cpp_line_break_before_comma_in_member_initializer_lists = true
-resharper_cpp_wrap_ternary_expr_style = chop_always
+resharper_cpp_wrap_ternary_expr_style = chop_if_long
 resharper_simple_block_style = line_break
 resharper_line_break_after_deref_in_trailing_return_types = on_single_line
 
