Git: Adjust ignore files for virtual production

#rnx

diff --git a/.gitignore b/.gitignore
--- a/.gitignore
+++ b/.gitignore
@@ -145,15 +145,14 @@
 /UE5CodeCompletionFolders.txt
 /UE5CodeLitePreProcessor.txt
 
-# Exception for VirtualProduction LiveLinkVCamFiles
+# Exception for VirtualProduction projects.
+!/Engine/Extras/**/*.dart
+!/Engine/Extras/**/*.kt
 !/Engine/Extras/VirtualProduction/LiveLinkVCAM/**
 /Engine/Extras/VirtualProduction/LiveLinkVCAM/LiveLinkProvider.framework/LiveLinkProvider
 /Engine/Extras/VirtualProduction/LiveLinkVCAM/UE4LiveLinkProvider.framework/UE4LiveLinkProvider
 /Engine/Extras/VirtualProduction/LiveLinkVCAM/UE5LiveLinkProvider.framework/UE5LiveLinkProvider
 
-# Exception for VP EpicStageApp dart files
-!/Engine/Extras/VirtualProduction/EpicStageApp/**/*.dart
-
 # Exception for VS Workspace patches
 !/Engine/Extras/VisualStudioWorkspace/
 
